- Fix static analysis warning about index out of bounds

diff --git a/Engine/Plugins/Runtime/OpenCV/Source/OpenCVHelper/Private/OpenCVHelper.cpp b/Engine/Plugins/Runtime/OpenCV/Source/OpenCVHelper/Private/OpenCVHelper.cpp
--- a/Engine/Plugins/Runtime/OpenCV/Source/OpenCVHelper/Private/OpenCVHelper.cpp
+++ b/Engine/Plugins/Runtime/OpenCV/Source/OpenCVHelper/Private/OpenCVHelper.cpp
@@ -224,6 +224,11 @@ double FOpenCVHelper::ComputeReprojectionError(const FTransform& CameraPose, con
 	// The 2D points will be compared against the undistorted 2D points, so the distortion coefficients can be ignored
 	cv::projectPoints(Points3d, Rrod, Tobj, CameraIntrinsicMatrix, cv::noArray(), ReprojectedPoints2d);
 
+	if (ReprojectedPoints2d.size() != NumPoints2d)
+	{
+		return -1.0;
+	}
+
 	// Compute euclidean distance between captured 2D points and reprojected 2D points to measure reprojection error
 	double ReprojectionError = 0.0;
 	for (int32 Index = 0; Index < NumPoints2d; ++Index)
