OpenCV helper Blueprint Function Library which adds convenience nodes for ArUco marker tracking and chessboard detection

diff --git a/Engine/Plugins/Runtime/OpenCV/Source/OpenCVHelper/Private/OpenCVBlueprintFunctionLibrary.cpp b/Engine/Plugins/Runtime/OpenCV/Source/OpenCVHelper/Private/OpenCVBlueprintFunctionLibrary.cpp
new file mode 100644
--- /dev/null
+++ b/Engine/Plugins/Runtime/OpenCV/Source/OpenCVHelper/Private/OpenCVBlueprintFunctionLibrary.cpp
@@ -0,0 +1,298 @@
+// Copyright Epic Games, Inc. All Rights Reserved.
+
+#include "OpenCVBlueprintFunctionLibrary.h"
+#include "OpenCVHelper.h"
+
+#include "Engine/Texture2D.h"
+#include "Engine/TextureRenderTarget2D.h"
+
+#if WITH_OPENCV
+
+#include <vector>
+
+OPENCV_INCLUDES_START
+#undef check 
+#include "opencv2/opencv.hpp"
+#include "opencv2/calib3d.hpp"
+#include "opencv2/aruco.hpp"
+OPENCV_INCLUDES_END
+
+#endif //WITH_OPENCV
+
+DEFINE_LOG_CATEGORY(LogOpenCVBlueprintFunctionLibrary);
+
+int32 UOpenCVBlueprintFunctionLibrary::OpenCVChessboardDetectCorners(const UTextureRenderTarget2D* InRenderTarget, const FIntPoint InPatternSize, const bool bDebugDrawCorners, UTexture2D*& OutDebugTexture, TArray<FVector2D>& OutDetectedCorners)
+{
+	if (InRenderTarget == nullptr)
+	{
+		UE_LOG(LogOpenCVBlueprintFunctionLibrary, Error, TEXT("OpenCVChessboardDetectCorners: Invalid InRenderTarget"));
+		return 0;
+	}
+
+	if (InRenderTarget->RenderTargetFormat != ETextureRenderTargetFormat::RTF_RGBA8)
+	{
+		UE_LOG(LogOpenCVBlueprintFunctionLibrary, Error, TEXT("OpenCVChessboardDetectCorners: Invalid InRenderTarget texture format - only RGBA8 is currently supported"));
+		return 0;
+	}
+
+	FRenderTarget* RenderTarget = ((UTextureRenderTarget2D*)InRenderTarget)->GameThread_GetRenderTargetResource();
+	if (RenderTarget == nullptr)
+	{
+		UE_LOG(LogOpenCVBlueprintFunctionLibrary, Error, TEXT("OpenCVChessboardDetectCorners: Can't get render target resource from InRenderTarget"));
+		return 0;
+	}
+
+	// Copy the pixels from the RenderTarget
+	TArray<FColor> Pixels;
+	if (!RenderTarget->ReadPixels(Pixels))
+	{
+		UE_LOG(LogOpenCVBlueprintFunctionLibrary, Error, TEXT("OpenCVChessboardDetectCorners: Failed to read pixels from RenderTarget"));
+		return 0;
+	}
+
+	// Check for invalid dimensions
+	const FIntPoint Size = RenderTarget->GetSizeXY();
+	if ((Size.X <= 0) || (Size.Y <= 0) || (Pixels.Num() != (Size.X * Size.Y)))
+	{
+		UE_LOG(LogOpenCVBlueprintFunctionLibrary, Error, TEXT("OpenCVChessboardDetectCorners: Invalid number of pixels read from RenderTarget"));
+		return 0;
+	}
+
+	// Create OpenCV Mat with those pixels and convert to grayscale for image processing
+	cv::Mat CvFrameColor(cv::Size(Size.X, Size.Y), CV_8UC4, Pixels.GetData());
+	cv::Mat CvFrameGray;
+	cv::cvtColor(CvFrameColor, CvFrameGray, cv::COLOR_RGBA2GRAY);
+
+	// Setup the incoming pattern size and flags
+	const cv::Size PatternSize(InPatternSize.X, InPatternSize.Y);
+	const int32 Flags = cv::CALIB_CB_ADAPTIVE_THRESH | cv::CALIB_CB_NORMALIZE_IMAGE | cv::CALIB_CB_FAST_CHECK;
+	std::vector<cv::Point2f> Corners;
+
+	// Detect the chessboard in the image
+	bool bAllCornersFound = cv::findChessboardCorners(CvFrameGray, PatternSize, Corners, Flags);
+	const int32 NumCornersFound = Corners.size();
+
+	// Draw debug info onto the frame if requested even if we don't find all the corners
+	if (bDebugDrawCorners)
+	{
+		cv::drawChessboardCorners(CvFrameColor, PatternSize, Corners, bAllCornersFound);
+		OutDebugTexture = FOpenCVHelper::TextureFromCvMat(CvFrameColor);
+	}
+
+	if (bAllCornersFound)
+	{
+		// Not completely sure why WinSize is set to these specific values, but every OpenCV example that I could find online just used these values
+		// ZeroZone of (-1, -1) means there is no dead zone in the search kernel
+		cv::Size WinSize(11, 11);
+		cv::Size ZeroZone(-1, -1);
+		cv::TermCriteria TermCriteria(cv::TermCriteria::Type::COUNT | cv::TermCriteria::Type::EPS, 30, 0.001);
+
+		// Refine the corner locations
+		cv::cornerSubPix(CvFrameGray, Corners, WinSize, ZeroZone, TermCriteria);
+
+		OutDetectedCorners.Reserve(NumCornersFound);
+		for (const cv::Point2f& Corner : Corners)
+		{
+			OutDetectedCorners.Add(FVector2D(Corner.x, Corner.y));
+		}
+	}
+
+	return NumCornersFound;
+}
+
+int32 UOpenCVBlueprintFunctionLibrary::OpenCVArucoDetectMarkers(const UTextureRenderTarget2D* InRenderTarget, const EOpenCVArucoDictionary InDictionary, const EOpenCVArucoDictionarySize InDictionarySize, const bool bDebugDrawMarkers, const bool bEstimatePose, const float InMarkerLengthInMeters, const FOpenCVLensDistortionParametersBase& InLensDistortionParameters, UTexture2D*& OutDebugTexture, TArray<FOpenCVArucoDetectedMarker>& OutDetectedMarkers)
+{
+	if (InRenderTarget == nullptr)
+	{
+		UE_LOG(LogOpenCVBlueprintFunctionLibrary, Error, TEXT("OpenCVArucoDetectMarkers: Invalid InRenderTarget"));
+		return 0;
+	}
+
+	if (bEstimatePose)
+	{
+		if (InMarkerLengthInMeters <= 0.0)
+		{
+			UE_LOG(LogOpenCVBlueprintFunctionLibrary, Error, TEXT("OpenCVArucoDetectMarkers: bEstimatePose specified, but InMarkerLengthInMeters is <= 0"));
+			return 0;
+		}
+	}
+
+	if (InRenderTarget->RenderTargetFormat != ETextureRenderTargetFormat::RTF_RGBA8)
+	{
+		UE_LOG(LogOpenCVBlueprintFunctionLibrary, Error, TEXT("OpenCVArucoDetectMarkers: Invalid InRenderTarget texture format - only RGBA8 is currently supported"));
+		return 0;
+	}
+
+	FRenderTarget* RenderTarget = ((UTextureRenderTarget2D*)InRenderTarget)->GameThread_GetRenderTargetResource();
+	if (RenderTarget == nullptr)
+	{
+		UE_LOG(LogOpenCVBlueprintFunctionLibrary, Error, TEXT("OpenCVArucoDetectMarkers: Can't get render target resource from InRenderTarget"));
+		return 0;
+	}
+
+	// Copy the pixels from the RenderTarget
+	TArray<FColor> Pixels;
+	if (!RenderTarget->ReadPixels(Pixels))
+	{
+		UE_LOG(LogOpenCVBlueprintFunctionLibrary, Error, TEXT("OpenCVArucoDetectMarkers: Failed to read pixels from RenderTarget"));
+		return 0;
+	}
+
+	// Check for invalid dimensions
+	const FIntPoint Size = RenderTarget->GetSizeXY();
+	if ((Size.X <= 0) || (Size.Y <= 0) || (Pixels.Num() != (Size.X * Size.Y)))
+	{
+		UE_LOG(LogOpenCVBlueprintFunctionLibrary, Error, TEXT("OpenCVArucoDetectMarkers: Invalid number of pixels read from RenderTarget"));
+		return 0;
+	}
+
+	// Create OpenCV Mat with those pixels and convert to grayscale for image processing
+	cv::Mat CvFrameColor(cv::Size(Size.X, Size.Y), CV_8UC4, Pixels.GetData());
+	cv::Mat CvFrameGray;
+	cv::cvtColor(CvFrameColor, CvFrameGray, cv::COLOR_RGBA2GRAY);
+
+	// Select the requested marker dictionary
+	cv::aruco::PREDEFINED_DICTIONARY_NAME DictionaryName = cv::aruco::DICT_ARUCO_ORIGINAL;
+	switch (InDictionary)
+	{
+		case EOpenCVArucoDictionary::Dict4x4:		switch (InDictionarySize)
+													{
+														case EOpenCVArucoDictionarySize::DictSize50:	DictionaryName = cv::aruco::DICT_4X4_50; break;
+														case EOpenCVArucoDictionarySize::DictSize100:	DictionaryName = cv::aruco::DICT_4X4_100; break;
+														case EOpenCVArucoDictionarySize::DictSize250:	DictionaryName = cv::aruco::DICT_4X4_250; break;
+														case EOpenCVArucoDictionarySize::DictSize1000:	DictionaryName = cv::aruco::DICT_4X4_1000; break;
+													}
+													break;
+		case EOpenCVArucoDictionary::Dict5x5:		switch (InDictionarySize)
+													{
+														case EOpenCVArucoDictionarySize::DictSize50:	DictionaryName = cv::aruco::DICT_5X5_50; break;
+														case EOpenCVArucoDictionarySize::DictSize100:	DictionaryName = cv::aruco::DICT_5X5_100; break;
+														case EOpenCVArucoDictionarySize::DictSize250:	DictionaryName = cv::aruco::DICT_5X5_250; break;
+														case EOpenCVArucoDictionarySize::DictSize1000:	DictionaryName = cv::aruco::DICT_5X5_1000; break;
+													}
+													break;
+		case EOpenCVArucoDictionary::Dict6x6:		switch (InDictionarySize)
+													{
+														case EOpenCVArucoDictionarySize::DictSize50:	DictionaryName = cv::aruco::DICT_6X6_50; break;
+														case EOpenCVArucoDictionarySize::DictSize100:	DictionaryName = cv::aruco::DICT_6X6_100; break;
+														case EOpenCVArucoDictionarySize::DictSize250:	DictionaryName = cv::aruco::DICT_6X6_250; break;
+														case EOpenCVArucoDictionarySize::DictSize1000:	DictionaryName = cv::aruco::DICT_6X6_1000; break;
+													}
+													break;
+		case EOpenCVArucoDictionary::Dict7x7:		switch (InDictionarySize)
+													{
+														case EOpenCVArucoDictionarySize::DictSize50:	DictionaryName = cv::aruco::DICT_7X7_50; break;
+														case EOpenCVArucoDictionarySize::DictSize100:	DictionaryName = cv::aruco::DICT_7X7_100; break;
+														case EOpenCVArucoDictionarySize::DictSize250:	DictionaryName = cv::aruco::DICT_7X7_250; break;
+														case EOpenCVArucoDictionarySize::DictSize1000:	DictionaryName = cv::aruco::DICT_7X7_1000; break;
+													}
+													break;
+		case EOpenCVArucoDictionary::DictOriginal:	DictionaryName = cv::aruco::DICT_ARUCO_ORIGINAL; break;
+	}
+	cv::Ptr<cv::aruco::Dictionary> Dictionary = cv::aruco::getPredefinedDictionary(DictionaryName);
+
+	// Find all ArUco markers
+	std::vector<int> MarkerIds;
+	std::vector<std::vector<cv::Point2f>> MarkerCorners;
+	const cv::Ptr<cv::aruco::DetectorParameters> DetectorParameters = cv::aruco::DetectorParameters::create();
+	cv::aruco::detectMarkers(CvFrameGray, Dictionary, MarkerCorners, MarkerIds, DetectorParameters);
+
+	static const uint32 kMaxCorners = 4;
+	const uint32 NumMarkersFound = MarkerIds.size();
+	if (NumMarkersFound > 0)
+	{
+		// Validate all return data
+		if (NumMarkersFound != MarkerCorners.size())
+		{
+			UE_LOG(LogOpenCVBlueprintFunctionLibrary, Error, TEXT("OpenCVArucoDetectMarkers: Number of detected MarkerIds and MarkerCorners don't match"));
+			return 0;
+		}
+
+		for (const std::vector<cv::Point2f>& Corner : MarkerCorners)
+		{
+			if (Corner.size() != kMaxCorners)
+			{
+				UE_LOG(LogOpenCVBlueprintFunctionLibrary, Error, TEXT("OpenCVArucoDetectMarkers: Number of detected corners per marker must be exactly %d"), kMaxCorners);
+				return 0;
+			}
+		}
+
+		// Get the estimated 3D pose data relative to the camera if requested
+		std::vector<cv::Vec3d> RotationVectors;
+		std::vector<cv::Vec3d> TranslationVectors;
+		cv::Mat CameraMatrix = cv::Mat::eye(3, 3, CV_64F);
+		cv::Mat DistortionCoeffs = cv::Mat::zeros(1, 8, CV_64F);
+		if (bEstimatePose)
+		{
+			CameraMatrix = InLensDistortionParameters.CreateOpenCVCameraMatrix(FVector2D(Size));
+			DistortionCoeffs = InLensDistortionParameters.ConvertToOpenCVDistortionCoefficients();
+
+			cv::aruco::estimatePoseSingleMarkers(MarkerCorners, InMarkerLengthInMeters, CameraMatrix, DistortionCoeffs, RotationVectors, TranslationVectors);
+		}
+
+		// Draw debug info onto the frame if requested
+		if (bDebugDrawMarkers)
+		{
+			// We need to convert to RGB temporarily because drawDetectedMarkers only supports grayscale or RGB
+			cv::cvtColor(CvFrameColor, CvFrameColor, cv::COLOR_RGBA2RGB);
+			cv::aruco::drawDetectedMarkers(CvFrameColor, MarkerCorners, MarkerIds, cv::Scalar(0, 255, 255));
+
+			if (bEstimatePose)
+			{
+				static const float AxisLengthInMeters = 0.05;
+				for (int32 Index = 0; Index < RotationVectors.size(); Index++)
+				{
+					// drawAxis only does one at a time - it will crash if more than one RotationVector or TranslationVector is passed in
+					cv::aruco::drawAxis(CvFrameColor, CameraMatrix, DistortionCoeffs, RotationVectors[Index], TranslationVectors[Index], AxisLengthInMeters);
+				}
+			}
+
+			cv::cvtColor(CvFrameColor, CvFrameColor, cv::COLOR_RGB2RGBA);
+
+			OutDebugTexture = FOpenCVHelper::TextureFromCvMat(CvFrameColor);
+		}
+
+		// Copy the OpenCV data into the Unreal types for return
+		OutDetectedMarkers.Reserve(NumMarkersFound);
+		for (uint32 Index = 0; Index < NumMarkersFound; Index++)
+		{
+			FOpenCVArucoDetectedMarker NewMarker;
+			NewMarker.Id = MarkerIds[Index];
+			NewMarker.Corners.Reserve(kMaxCorners);
+			for (const cv::Point2f& Corner : MarkerCorners[Index])
+			{
+				// Normalize each corner to the input image size
+				NewMarker.Corners.Add(FVector2D(Corner.x / Size.X, Corner.y / Size.Y));
+			}
+
+			// Convert the returned pose into the Unreal coordinate system
+			FTransform NewPose = FTransform::Identity;
+			if (bEstimatePose)
+			{
+				// Convert the Rodrigues rotation into a 3x3 rotation matrix
+				cv::Mat RotationMatrix = cv::Mat::eye(3, 3, CV_64F);
+				cv::Rodrigues(RotationVectors[Index], RotationMatrix);
+
+				// Populate the Unreal rotation matrix
+				FMatrix UnrealMatrix = FMatrix::Identity;
+				for (int32 Column = 0; Column < 3; Column++)
+				{
+					UnrealMatrix.SetColumn(Column, FVector(RotationMatrix.at<double>(Column, 0), RotationMatrix.at<double>(Column, 1), RotationMatrix.at<double>(Column, 2)));
+				}
+
+				// Scale our translation and populate the FTransform
+				static const double MetersToUnrealUnits = 100.0;
+				FVector Translation(TranslationVectors[Index][0], TranslationVectors[Index][1], TranslationVectors[Index][2]);
+				NewPose = FTransform(UnrealMatrix.GetColumn(0), UnrealMatrix.GetColumn(1), UnrealMatrix.GetColumn(2), Translation * MetersToUnrealUnits);
+
+				// Convert to Unreal axes
+				FOpenCVHelper::ConvertOpenCVToUnreal(NewPose);
+			}
+
+			NewMarker.Pose = NewPose;
+			OutDetectedMarkers.Add(NewMarker);
+		}
+	}
+
+	return NumMarkersFound;
+}
diff --git a/Engine/Plugins/Runtime/OpenCV/Source/OpenCVHelper/Public/OpenCVBlueprintFunctionLibrary.h b/Engine/Plugins/Runtime/OpenCV/Source/OpenCVHelper/Public/OpenCVBlueprintFunctionLibrary.h
new file mode 100644
--- /dev/null
+++ b/Engine/Plugins/Runtime/OpenCV/Source/OpenCVHelper/Public/OpenCVBlueprintFunctionLibrary.h
@@ -0,0 +1,85 @@
+// Copyright Epic Games, Inc. All Rights Reserved.
+
+#pragma once
+
+#include "Kismet/BlueprintFunctionLibrary.h"
+
+#include "OpenCVHelper.h"
+
+#include "OpenCVBlueprintFunctionLibrary.generated.h"
+
+DECLARE_LOG_CATEGORY_EXTERN(LogOpenCVBlueprintFunctionLibrary, Log, All);
+
+class UTexture2D;
+class UTextureRenderTarget2D;
+
+UENUM(BlueprintType)
+enum class EOpenCVArucoDictionary : uint8
+{
+	Dict4x4 UMETA(DisplayName = "4x4"),
+	Dict5x5 UMETA(DisplayName = "5x5"),
+	Dict6x6 UMETA(DisplayName = "6x6"),
+	Dict7x7 UMETA(DisplayName = "7x7"),
+	DictOriginal UMETA(DisplayName = "Original")
+};
+
+UENUM(BlueprintType)
+enum class EOpenCVArucoDictionarySize : uint8
+{
+	DictSize50 UMETA(DisplayName = "50"),
+	DictSize100 UMETA(DisplayName = "100"),
+	DictSize250 UMETA(DisplayName = "250"),
+	DictSize1000 UMETA(DisplayName = "1000")
+};
+
+USTRUCT(BlueprintType)
+struct FOpenCVArucoDetectedMarker
+{
+	GENERATED_BODY()
+
+	UPROPERTY(EditAnywhere, BlueprintReadWrite, Category = "OpenCV | ArUco")
+	int32 Id;
+
+	UPROPERTY(EditAnywhere, BlueprintReadWrite, Category = "OpenCV | ArUco")
+	TArray<FVector2D> Corners;
+
+	UPROPERTY(EditAnywhere, BlueprintReadWrite, Category = "OpenCV | ArUco")
+	FTransform Pose;
+};
+
+UCLASS(BlueprintType)
+class OPENCVHELPER_API UOpenCVBlueprintFunctionLibrary : public UBlueprintFunctionLibrary
+{
+	GENERATED_BODY()
+
+public:
+	/**
+	 * Detects a camera calibration chessboard in the supplied image
+	 *
+	 * @param InRenderTarget Input image in which to search for a chessboard
+	 * @param InPatternSize Number of interior corners on the physical chessboard (rows, columns)
+	 * @param bDebugDrawCorners If true, output a Texture2D showing the detected corner debug info overlaid on the input image
+	 * @param OutDebugTexture Output debug image (required if bDebugDrawCorners is True)
+	 * @param OutDetectedCorners Output array of corners detected in the input image
+	 * @return Total number of corners detected in the input image
+	 */
+	UFUNCTION(BlueprintCallable, Category = "OpenCV", Meta = (DisplayName = "OpenCV Chessboard Detect Corners"))
+	static int32 OpenCVChessboardDetectCorners(const UTextureRenderTarget2D* InRenderTarget, const FIntPoint InPatternSize, const bool bDebugDrawCorners, UTexture2D*& OutDebugTexture, TArray<FVector2D>& OutDetectedCorners);
+
+	/**
+	 * Detects all ArUco markers in the supplied image
+	 *
+	 * @param InRenderTarget Input image in which to search for markers
+	 * @param InDictionary Which ArUco marker dictionary to use for detection
+	 * @param InDictionarySize The size of the ArUco marker dictionary
+	 * @param bDebugDrawMarkers If true, output a Texture2D showing the detected marker debug info overlaid on the input image
+	 * @param bEstimatePose If true, return the 3D pose for each marker relative to the camera position
+	 * @param InMarkerLengthInMeters Length in meters of one side of the physical marker (required if bEstimatePose is True)
+	 * @param InCameraCalibrationParameters Lens distortion parameters for the incoming image (required if bEstimatePose is True)
+	 * @param OutDebugTexture Output debug image (required if bDebugDrawMarkers is True)
+	 * @param OutDetectedMarkers Output array of markers detected in the input image
+	 * @return Total number of markers detected in the input image
+	 */
+	UFUNCTION(BlueprintCallable, Category = "OpenCV | ArUco", Meta = (DisplayName = "OpenCV ArUco Detect Markers"))
+	static int32 OpenCVArucoDetectMarkers(const UTextureRenderTarget2D* InRenderTarget, const EOpenCVArucoDictionary InDictionary, const EOpenCVArucoDictionarySize InDictionarySize, const bool bDebugDrawMarkers, const bool bEstimatePose, const float InMarkerLengthInMeters, const FOpenCVLensDistortionParametersBase& InLensDistortionParameters, UTexture2D*& OutDebugTexture, TArray<FOpenCVArucoDetectedMarker>& OutDetectedMarkers);
+};
\ No newline at end of file
