Fixed OpenCVHelper Linux errors

diff --git a/Engine/Plugins/Runtime/OpenCV/Source/OpenCVHelper/Private/OpenCVBlueprintFunctionLibrary.cpp b/Engine/Plugins/Runtime/OpenCV/Source/OpenCVHelper/Private/OpenCVBlueprintFunctionLibrary.cpp
--- a/Engine/Plugins/Runtime/OpenCV/Source/OpenCVHelper/Private/OpenCVBlueprintFunctionLibrary.cpp
+++ b/Engine/Plugins/Runtime/OpenCV/Source/OpenCVHelper/Private/OpenCVBlueprintFunctionLibrary.cpp
@@ -17,12 +17,13 @@ OPENCV_INCLUDES_START
 #include "opencv2/aruco.hpp"
 OPENCV_INCLUDES_END
 
-#endif //WITH_OPENCV
+#endif	// WITH_OPENCV
 
 DEFINE_LOG_CATEGORY(LogOpenCVBlueprintFunctionLibrary);
 
 int32 UOpenCVBlueprintFunctionLibrary::OpenCVChessboardDetectCorners(const UTextureRenderTarget2D* InRenderTarget, const FIntPoint InPatternSize, const bool bDebugDrawCorners, UTexture2D*& OutDebugTexture, TArray<FVector2D>& OutDetectedCorners)
 {
+#if WITH_OPENCV
 	if (InRenderTarget == nullptr)
 	{
 		UE_LOG(LogOpenCVBlueprintFunctionLibrary, Error, TEXT("OpenCVChessboardDetectCorners: Invalid InRenderTarget"));
@@ -98,10 +99,17 @@ int32 UOpenCVBlueprintFunctionLibrary::OpenCVChessboardDetectCorners(const UText
 	}
 
 	return NumCornersFound;
+#else
+	{
+		UE_LOG(LogOpenCVBlueprintFunctionLibrary, Error, TEXT("OpenCVChessboardDetectCorners: Requires OpenCV"));
+		return 0;
+	}
+#endif	// WITH_OPENCV
 }
 
 int32 UOpenCVBlueprintFunctionLibrary::OpenCVArucoDetectMarkers(const UTextureRenderTarget2D* InRenderTarget, const EOpenCVArucoDictionary InDictionary, const EOpenCVArucoDictionarySize InDictionarySize, const bool bDebugDrawMarkers, const bool bEstimatePose, const float InMarkerLengthInMeters, const FOpenCVLensDistortionParametersBase& InLensDistortionParameters, UTexture2D*& OutDebugTexture, TArray<FOpenCVArucoDetectedMarker>& OutDetectedMarkers)
 {
+#if WITH_OPENCV
 	if (InRenderTarget == nullptr)
 	{
 		UE_LOG(LogOpenCVBlueprintFunctionLibrary, Error, TEXT("OpenCVArucoDetectMarkers: Invalid InRenderTarget"));
@@ -198,7 +206,7 @@ int32 UOpenCVBlueprintFunctionLibrary::OpenCVArucoDetectMarkers(const UTextureRe
 	cv::aruco::detectMarkers(CvFrameGray, Dictionary, MarkerCorners, MarkerIds, DetectorParameters);
 
 	static const uint32 kMaxCorners = 4;
-	const uint32 NumMarkersFound = MarkerIds.size();
+	const int32 NumMarkersFound = MarkerIds.size();
 	if (NumMarkersFound > 0)
 	{
 		// Validate all return data
@@ -295,4 +303,10 @@ int32 UOpenCVBlueprintFunctionLibrary::OpenCVArucoDetectMarkers(const UTextureRe
 	}
 
 	return NumMarkersFound;
+#else
+	{
+		UE_LOG(LogOpenCVBlueprintFunctionLibrary, Error, TEXT("OpenCVArucoDetectMarkers: Requires OpenCV"));
+		return 0;
+	}
+#endif	// WITH_OPENCV
 }
