Fixed build error that preflight minimal did not catch :(

diff --git a/Engine/Plugins/Runtime/OpenCV/Source/OpenCVHelper/Private/OpenCVBlueprintFunctionLibrary.cpp b/Engine/Plugins/Runtime/OpenCV/Source/OpenCVHelper/Private/OpenCVBlueprintFunctionLibrary.cpp
--- a/Engine/Plugins/Runtime/OpenCV/Source/OpenCVHelper/Private/OpenCVBlueprintFunctionLibrary.cpp
+++ b/Engine/Plugins/Runtime/OpenCV/Source/OpenCVHelper/Private/OpenCVBlueprintFunctionLibrary.cpp
@@ -262,7 +262,7 @@ int32 UOpenCVBlueprintFunctionLibrary::OpenCVArucoDetectMarkers(const UTextureRe
 
 		// Copy the OpenCV data into the Unreal types for return
 		OutDetectedMarkers.Reserve(NumMarkersFound);
-		for (uint32 Index = 0; Index < NumMarkersFound; Index++)
+		for (int32 Index = 0; Index < NumMarkersFound; Index++)
 		{
 			FOpenCVArucoDetectedMarker NewMarker;
 			NewMarker.Id = MarkerIds[Index];
