Fixed default initializer for FOpenCVArucoDetectedMarker

diff --git a/Engine/Plugins/Runtime/OpenCV/Source/OpenCVHelper/Public/OpenCVBlueprintFunctionLibrary.h b/Engine/Plugins/Runtime/OpenCV/Source/OpenCVHelper/Public/OpenCVBlueprintFunctionLibrary.h
--- a/Engine/Plugins/Runtime/OpenCV/Source/OpenCVHelper/Public/OpenCVBlueprintFunctionLibrary.h
+++ b/Engine/Plugins/Runtime/OpenCV/Source/OpenCVHelper/Public/OpenCVBlueprintFunctionLibrary.h
@@ -37,6 +37,11 @@ struct FOpenCVArucoDetectedMarker
 {
 	GENERATED_BODY()
 
+	FOpenCVArucoDetectedMarker()
+		: Id(-1)
+	{
+	}
+
 	UPROPERTY(EditAnywhere, BlueprintReadWrite, Category = "OpenCV | ArUco")
 	int32 Id;
 
