CameraCalibration: Refactor opencv implementation details out of the camera calibration plugins and into the OpenCV Helper module

diff --git a/Engine/Plugins/Experimental/VirtualProduction/LedWallCalibration/Source/LedWallCalibration/Public/LedWallArucoGenerationOptions.h b/Engine/Plugins/Experimental/VirtualProduction/LedWallCalibration/Source/LedWallCalibration/Public/LedWallArucoGenerationOptions.h
--- a/Engine/Plugins/Experimental/VirtualProduction/LedWallCalibration/Source/LedWallCalibration/Public/LedWallArucoGenerationOptions.h
+++ b/Engine/Plugins/Experimental/VirtualProduction/LedWallCalibration/Source/LedWallCalibration/Public/LedWallArucoGenerationOptions.h
@@ -3,32 +3,10 @@
 #pragma once
 
 #include "CoreMinimal.h"
+#include "OpenCVHelper.h"
 
 #include "LedWallArucoGenerationOptions.generated.h"
 
-/** Enumeration of known Aruco dictionaries */
-UENUM(BlueprintType)
-enum EArucoDictionary : int
-{
-	DICT_4X4_50          UMETA(DisplayName = "DICT_4X4_50"),
-	DICT_4X4_100		 UMETA(DisplayName = "DICT_4X4_100"),
-	DICT_4X4_250		 UMETA(DisplayName = "DICT_4X4_250"),
-	DICT_4X4_1000		 UMETA(DisplayName = "DICT_4X4_1000"),
-	DICT_5X5_50			 UMETA(DisplayName = "DICT_5X5_50"),
-	DICT_5X5_100		 UMETA(DisplayName = "DICT_5X5_100"),
-	DICT_5X5_250		 UMETA(DisplayName = "DICT_5X5_250"),
-	DICT_5X5_1000		 UMETA(DisplayName = "DICT_5X5_1000"),
-	DICT_6X6_50			 UMETA(DisplayName = "DICT_6X6_50"),
-	DICT_6X6_100		 UMETA(DisplayName = "DICT_6X6_100"),
-	DICT_6X6_250		 UMETA(DisplayName = "DICT_6X6_250"),
-	DICT_6X6_1000		 UMETA(DisplayName = "DICT_6X6_1000"),
-	DICT_7X7_50			 UMETA(DisplayName = "DICT_7X7_50"),
-	DICT_7X7_100		 UMETA(DisplayName = "DICT_7X7_100"),
-	DICT_7X7_250		 UMETA(DisplayName = "DICT_7X7_250"),
-	DICT_7X7_1000		 UMETA(DisplayName = "DICT_7X7_1000"),
-	DICT_ARUCO_ORIGINAL	 UMETA(DisplayName = "DICT_ARUCO_ORIGINAL")
-};
-
 /**
  * Structure that can be passed to the Aruco generation function
  */
@@ -49,7 +27,7 @@ public:
 
 	/** Aruco dictionary to use when generating the markers */
 	UPROPERTY(EditAnywhere, BlueprintReadWrite, Category = "Aruco")
-	TEnumAsByte<EArucoDictionary> ArucoDictionary = EArucoDictionary::DICT_6X6_1000;
+	EArucoDictionary ArucoDictionary = EArucoDictionary::DICT_6X6_1000;
 
 	/** Starting marker Id. Arucos will be generated with consecutive Marker Ids, starting from this one */
 	UPROPERTY(EditAnywhere, BlueprintReadWrite, Category = "Aruco")
@@ -71,13 +49,13 @@ public:
 	 */
 	FString ArucoDictionaryAsString() const
 	{
-		const UEnum* EnumPtr = FindObject<UEnum>(nullptr, TEXT("/Script/LedWallCalibration.EArucoDictionary"), true);
+		const UEnum* EnumPtr = FindObject<UEnum>(nullptr, TEXT("/Script/OpenCVHelper.EArucoDictionary"), true);
 		check(EnumPtr);
 
 #if WITH_EDITOR
-		return EnumPtr->GetDisplayNameTextByIndex(static_cast<uint8>(ArucoDictionary.GetValue())).ToString();
+		return EnumPtr->GetDisplayNameTextByIndex(static_cast<uint8>(ArucoDictionary)).ToString();
 #else
-		return EnumPtr->GetNameStringByIndex(static_cast<uint8>(ArucoDictionary.GetValue()));
+		return EnumPtr->GetNameStringByIndex(static_cast<uint8>(ArucoDictionary));
 #endif
 	}
 };
diff --git a/Engine/Plugins/Experimental/VirtualProduction/LedWallCalibration/Source/LedWallCalibrationEditor/Private/CalibrationPointArucosForWallDetailsRow.h b/Engine/Plugins/Experimental/VirtualProduction/LedWallCalibration/Source/LedWallCalibrationEditor/Private/CalibrationPointArucosForWallDetailsRow.h
--- a/Engine/Plugins/Experimental/VirtualProduction/LedWallCalibration/Source/LedWallCalibrationEditor/Private/CalibrationPointArucosForWallDetailsRow.h
+++ b/Engine/Plugins/Experimental/VirtualProduction/LedWallCalibration/Source/LedWallCalibrationEditor/Private/CalibrationPointArucosForWallDetailsRow.h
@@ -25,7 +25,7 @@ private:
 private:
 
 	/** Remember last Aruco dictionary used so that it can be pre-chosen next time */
-	TEnumAsByte<EArucoDictionary> PreviousArucoDictionaryUsed = EArucoDictionary::DICT_6X6_1000;
+	EArucoDictionary PreviousArucoDictionaryUsed = EArucoDictionary::DICT_6X6_1000;
 
 	/** 
 	 * Remember next marker id so that it can pre-choose the next starting marker id. 
diff --git a/Engine/Plugins/Runtime/OpenCV/Source/OpenCVHelper/Private/OpenCVHelper.cpp b/Engine/Plugins/Runtime/OpenCV/Source/OpenCVHelper/Private/OpenCVHelper.cpp
--- a/Engine/Plugins/Runtime/OpenCV/Source/OpenCVHelper/Private/OpenCVHelper.cpp
+++ b/Engine/Plugins/Runtime/OpenCV/Source/OpenCVHelper/Private/OpenCVHelper.cpp
@@ -9,12 +9,68 @@
 #if WITH_OPENCV
 
 #include "PreOpenCVHeaders.h" // IWYU pragma: keep
-
+#include "opencv2/aruco.hpp"
 #include "opencv2/calib3d.hpp"
+#include "opencv2/imgproc.hpp"
 #include "PostOpenCVHeaders.h" // IWYU pragma: keep
 
 #endif	// WITH_OPENCV
 
+
+DEFINE_LOG_CATEGORY_STATIC(LogOpenCVHelper, Log, All);
+
+
+namespace UE::Aruco::Private
+{
+#if WITH_OPENCV
+	cv::Ptr<cv::aruco::Dictionary> GetArucoDictionary(EArucoDictionary InDictionary)
+	{
+		switch (InDictionary)
+		{
+		case EArucoDictionary::DICT_4X4_50:
+			return cv::aruco::getPredefinedDictionary(cv::aruco::DICT_4X4_50);		
+		case EArucoDictionary::DICT_4X4_100:
+			return cv::aruco::getPredefinedDictionary(cv::aruco::DICT_4X4_100);
+		case EArucoDictionary::DICT_4X4_250:
+			return cv::aruco::getPredefinedDictionary(cv::aruco::DICT_4X4_250);
+		case EArucoDictionary::DICT_4X4_1000:
+			return cv::aruco::getPredefinedDictionary(cv::aruco::DICT_4X4_1000);
+		case EArucoDictionary::DICT_5X5_50:
+			return cv::aruco::getPredefinedDictionary(cv::aruco::DICT_5X5_50);
+		case EArucoDictionary::DICT_5X5_100:
+			return cv::aruco::getPredefinedDictionary(cv::aruco::DICT_5X5_100);
+		case EArucoDictionary::DICT_5X5_250:
+			return cv::aruco::getPredefinedDictionary(cv::aruco::DICT_5X5_250);
+		case EArucoDictionary::DICT_5X5_1000:
+			return cv::aruco::getPredefinedDictionary(cv::aruco::DICT_5X5_1000);
+		case EArucoDictionary::DICT_6X6_50:
+			return cv::aruco::getPredefinedDictionary(cv::aruco::DICT_6X6_50);
+		case EArucoDictionary::DICT_6X6_100:
+			return cv::aruco::getPredefinedDictionary(cv::aruco::DICT_6X6_100);
+		case EArucoDictionary::DICT_6X6_250:
+			return cv::aruco::getPredefinedDictionary(cv::aruco::DICT_6X6_250);
+		case EArucoDictionary::DICT_6X6_1000:
+			return cv::aruco::getPredefinedDictionary(cv::aruco::DICT_6X6_1000);
+		case EArucoDictionary::DICT_7X7_50:
+			return cv::aruco::getPredefinedDictionary(cv::aruco::DICT_7X7_50);
+		case EArucoDictionary::DICT_7X7_100:
+			return cv::aruco::getPredefinedDictionary(cv::aruco::DICT_7X7_100);
+		case EArucoDictionary::DICT_7X7_250:
+			return cv::aruco::getPredefinedDictionary(cv::aruco::DICT_7X7_250);
+		case EArucoDictionary::DICT_7X7_1000:
+			return cv::aruco::getPredefinedDictionary(cv::aruco::DICT_7X7_1000);
+		case EArucoDictionary::DICT_ARUCO_ORIGINAL:
+			return cv::aruco::getPredefinedDictionary(cv::aruco::DICT_ARUCO_ORIGINAL);
+		default:
+			ensureMsgf(false, TEXT("Unhandled EArucoDictionary type. Update this switch statement."));
+			break; // Do nothing
+		}
+
+		return nullptr;
+	}
+#endif
+}
+
 void FOpenCVHelper::ConvertCoordinateSystem(FTransform& Transform, const EAxis SrcXInDstAxis, const EAxis SrcYInDstAxis, const EAxis SrcZInDstAxis)
 {
 	// Unreal Engine:
@@ -46,43 +102,17 @@
 	ConvertCoordinateSystem(Transform, EAxis::Z, EAxis::X, EAxis::Yn);
 }
 
+FVector FOpenCVHelper::ConvertUnrealToOpenCV(const FVector& Vector)
+{
+	return FVector(Vector.Y, -Vector.Z, Vector.X);
+}
+
+FVector FOpenCVHelper::ConvertOpenCVToUnreal(const FVector& Vector)
+{
+	return FVector(Vector.Z, Vector.X, -Vector.Y);
+}
+
 #if WITH_OPENCV
-cv::Mat FOpenCVLensDistortionParametersBase::ConvertToOpenCVDistortionCoefficients() const
-{
-	if (bUseFisheyeModel)
-	{
-		cv::Mat DistortionCoefficients(1, 4, CV_64F);
-		DistortionCoefficients.at<double>(0) = K1;
-		DistortionCoefficients.at<double>(1) = K2;
-		DistortionCoefficients.at<double>(2) = K3;
-		DistortionCoefficients.at<double>(3) = K4;
-		return DistortionCoefficients;
-	}
-	else
-	{
-		cv::Mat DistortionCoefficients(1, 8, CV_64F);
-		DistortionCoefficients.at<double>(0) = K1;
-		DistortionCoefficients.at<double>(1) = K2;
-		DistortionCoefficients.at<double>(2) = P1;
-		DistortionCoefficients.at<double>(3) = P2;
-		DistortionCoefficients.at<double>(4) = K3;
-		DistortionCoefficients.at<double>(5) = K4;
-		DistortionCoefficients.at<double>(6) = K5;
-		DistortionCoefficients.at<double>(7) = K6;
-		return DistortionCoefficients;
-	}
-}
-
-cv::Mat FOpenCVLensDistortionParametersBase::CreateOpenCVCameraMatrix(const FVector2D& InImageSize) const
-{
-	cv::Mat CameraMatrix = cv::Mat::eye(3, 3, CV_64F);
-	CameraMatrix.at<double>(0, 0) = F.X * InImageSize.X;
-	CameraMatrix.at<double>(1, 1) = F.Y * InImageSize.Y;
-	CameraMatrix.at<double>(0, 2) = C.X * InImageSize.X;
-	CameraMatrix.at<double>(1, 2) = C.Y * InImageSize.Y;
-	return CameraMatrix;
-}
-
 UTexture2D* FOpenCVHelper::TextureFromCvMat(cv::Mat& Mat, const FString* PackagePath, const FName* TextureName)
 {
 	if ((Mat.cols <= 0) || (Mat.rows <= 0))
@@ -259,6 +289,364 @@
 	return InTexture;
 }
 
+void FOpenCVHelper::MakeCameraPoseFromObjectVectors(const cv::Mat& InRotation, const cv::Mat& InTranslation, FTransform& OutTransform)
+{
+	// The input rotation and translation vectors are both 3x1 vectors that are computed by OpenCV in functions such as calibrateCamera() and solvePnP()
+	// Together, they perform a change of basis from object coordinate space to camera coordinate space
+	// The desired output transform is the pose of the camera in world space, which can be obtained by inverting/transposing the rotation and translation vectors
+	// [R|t]' = [R'|-R'*t]
+	// Where R is a 3x3 rotation matrix and t is a 3x1 translation vector
+
+	// Convert input rotation from rodrigues notation to a 3x3 rotation matrix
+	cv::Mat RotationMatrix3x3;
+	cv::Rodrigues(InRotation, RotationMatrix3x3);
+
+	// Invert/transpose to get the camera orientation
+	cv::Mat CameraRotation = RotationMatrix3x3.t();
+	cv::Mat CameraTranslation = -RotationMatrix3x3.t() * InTranslation;
+
+	FMatrix TransformationMatrix = FMatrix::Identity;
+
+	// Add the rotation matrix to the transformation matrix
+	for (int32 Column = 0; Column < 3; ++Column)
+	{
+		TransformationMatrix.SetColumn(Column, FVector(CameraRotation.at<double>(Column, 0), CameraRotation.at<double>(Column, 1), CameraRotation.at<double>(Column, 2)));
+	}
+
+	// Add the translation vector to the transformation matrix
+	TransformationMatrix.M[3][0] = CameraTranslation.at<double>(0);
+	TransformationMatrix.M[3][1] = CameraTranslation.at<double>(1);
+	TransformationMatrix.M[3][2] = CameraTranslation.at<double>(2);
+
+	OutTransform.SetFromMatrix(TransformationMatrix);
+
+	// Convert the output FTransform to UE's coordinate system
+	FOpenCVHelper::ConvertOpenCVToUnreal(OutTransform);
+}
+
+void FOpenCVHelper::MakeObjectVectorsFromCameraPose(const FTransform& InTransform, cv::Mat& OutRotation, cv::Mat& OutTranslation)
+{
+	// Convert the input FTransform to OpenCV's coordinate system
+	FTransform CvTransform = InTransform;
+	ConvertUnrealToOpenCV(CvTransform);
+
+	const FMatrix TransformationMatrix = CvTransform.ToMatrixNoScale();
+
+	// Extract the translation vector from the transformation matrix
+	cv::Mat CameraTranslation = cv::Mat(3, 1, CV_64FC1);
+	CameraTranslation.at<double>(0) = TransformationMatrix.M[3][0];
+	CameraTranslation.at<double>(1) = TransformationMatrix.M[3][1];
+	CameraTranslation.at<double>(2) = TransformationMatrix.M[3][2];
+
+	// Extract the rotation matrix from the transformation matrix
+	cv::Mat CameraRotation = cv::Mat(3, 3, CV_64FC1);
+	for (int32 Column = 0; Column < 3; ++Column)
+	{
+		const FVector ColumnVector = TransformationMatrix.GetColumn(Column);
+		CameraRotation.at<double>(Column, 0) = ColumnVector.X;
+		CameraRotation.at<double>(Column, 1) = ColumnVector.Y;
+		CameraRotation.at<double>(Column, 2) = ColumnVector.Z;
+	}
+
+
+	// Invert/transpose to get the rotation and translation that perform a change of basis from object coordinate space to camera coordinate space
+	cv::Mat RotationMatrix3x3 = CameraRotation.t();
+	OutTranslation = -CameraRotation.inv() * CameraTranslation;
+
+	// Convert the 3x3 rotation matrix to rodrigues notation
+	cv::Rodrigues(RotationMatrix3x3, OutRotation);
+}
+#endif // WITH_OPENCV
+
+bool FOpenCVHelper::IdentifyArucoMarkers(TArray<FColor>& Image, FIntPoint ImageSize, EArucoDictionary DictionaryName, TArray<FArucoMarker>& OutMarkers)
+{
+#if WITH_OPENCV
+	// Initialize an OpenCV matrix header to point at the input image data. 
+	// The format for FColor is B8G8R8A8, so the corresponding OpenCV format is CV_8UC4 (8-bit per channel, 4 channels)
+	cv::Mat ImageMat = cv::Mat(ImageSize.Y, ImageSize.X, CV_8UC4, Image.GetData());
+
+	// Convert the image to grayscale before attempting to detect markers
+	cv::Mat GrayImage;
+	cv::cvtColor(ImageMat, GrayImage, cv::COLOR_RGBA2GRAY);
+
+	// We do not know the number of markers expected in the image, so we cannot yet reserve the necessary space for our marker data and create a convenient matrix header for that data. 
+	// Instead, we will pass unitialized vectors to opencv and copy the data into our final data structure after the markers have been identified. 
+ 	std::vector<int> MarkerIds;
+ 	std::vector<std::vector<cv::Point2f>> Corners;
+
+	cv::Ptr<cv::aruco::Dictionary> Dictionary = UE::Aruco::Private::GetArucoDictionary(DictionaryName);
+	if (!Dictionary)
+	{
+		return false;
+	}
+
+	cv::Ptr<cv::aruco::DetectorParameters> DetectorParameters = cv::aruco::DetectorParameters::create();
+	DetectorParameters->cornerRefinementMethod = cv::aruco::CORNER_REFINE_SUBPIX;
+
+	cv::aruco::detectMarkers(GrayImage, Dictionary, Corners, MarkerIds, DetectorParameters);
+
+	OutMarkers.Empty();
+
+	const int32 NumMarkers = MarkerIds.size();
+
+	if (NumMarkers < 1)
+	{
+		return false;
+	}
+
+	OutMarkers.Reserve(NumMarkers);
+
+	// Copy the detected marker data into the final data structure
+	for (int32 MarkerIndex = 0; MarkerIndex < NumMarkers; ++MarkerIndex)
+	{
+		FArucoMarker NewMarker;
+		NewMarker.MarkedId = MarkerIds[MarkerIndex];
+
+		const std::vector<cv::Point2f>& MarkerCorners = Corners[MarkerIndex];
+
+		NewMarker.Corners[0] = FVector2f(MarkerCorners[0].x, MarkerCorners[0].y);
+		NewMarker.Corners[1] = FVector2f(MarkerCorners[1].x, MarkerCorners[1].y);
+		NewMarker.Corners[2] = FVector2f(MarkerCorners[2].x, MarkerCorners[2].y);
+		NewMarker.Corners[3] = FVector2f(MarkerCorners[3].x, MarkerCorners[3].y);
+
+		OutMarkers.Add(NewMarker);
+	}
+
+	return true;
+#else
+	return false;
+#endif // WITH_OPENCV
+}
+
+bool FOpenCVHelper::DrawArucoMarkers(const TArray<FArucoMarker>& Markers, UTexture2D* DebugTexture)
+{
+#if WITH_OPENCV
+	if (!DebugTexture || Markers.IsEmpty())
+	{
+		return false;
+	}
+
+	const int32 NumMarkers = Markers.Num();
+
+	std::vector<int> MarkerIds;
+	MarkerIds.reserve(NumMarkers);
+
+	std::vector<std::vector<cv::Point2f>> Corners;
+	Corners.reserve(NumMarkers);
+
+	// Copy the detected marker data into the final data structure
+	for (int32 MarkerIndex = 0; MarkerIndex < NumMarkers; ++MarkerIndex)
+	{
+		const FArucoMarker& Marker = Markers[MarkerIndex];
+
+		MarkerIds.push_back(Marker.MarkedId);
+
+		std::vector<cv::Point2f> MarkerCorners;
+		MarkerCorners.reserve(4);
+
+		MarkerCorners.push_back(cv::Point2f(Marker.Corners[0].X, Marker.Corners[0].Y));
+		MarkerCorners.push_back(cv::Point2f(Marker.Corners[1].X, Marker.Corners[1].Y));
+		MarkerCorners.push_back(cv::Point2f(Marker.Corners[2].X, Marker.Corners[2].Y));
+		MarkerCorners.push_back(cv::Point2f(Marker.Corners[3].X, Marker.Corners[3].Y));
+
+		Corners.push_back(MarkerCorners);
+	}
+
+	FTexture2DMipMap& Mip0 = DebugTexture->GetPlatformData()->Mips[0];
+	void* TextureData = Mip0.BulkData.Lock(LOCK_READ_WRITE);
+
+	// Create a matrix header pointing to the raw texture data, copy the debug marker outlines into it, and update the texture resource
+	cv::Mat TextureMat = cv::Mat(DebugTexture->GetSizeY(), DebugTexture->GetSizeX(), CV_8UC4, TextureData);
+
+	// We need to convert from RGBA to RGB temporarily because cv::aruco::drawDetectedMarkers only supports images with 1 or 3 color channels (not 4)
+	// We use a second matrix header because cvtColor will make a copy of the original data each time it changes the color format and we do not want to lose the reference to the original TextureData pointer.
+	cv::Mat DebugMat;
+	TextureMat.copyTo(DebugMat);
+
+	cv::cvtColor(DebugMat, DebugMat, cv::COLOR_RGBA2RGB);
+	cv::aruco::drawDetectedMarkers(DebugMat, Corners, MarkerIds);
+	cv::cvtColor(DebugMat, DebugMat, cv::COLOR_RGB2RGBA);
+
+	DebugMat.copyTo(TextureMat);
+
+	Mip0.BulkData.Unlock();
+
+	DebugTexture->UpdateResource();
+
+	return true;
+#else
+	return false;
+#endif // WITH_OPENCV
+}
+
+bool FOpenCVHelper::IdentifyCheckerboard(TArray<FColor>& Image, FIntPoint ImageSize, FIntPoint CheckerboardDimensions, TArray<FVector2f>& OutCorners)
+{
+#if WITH_OPENCV
+	// Initialize an OpenCV matrix header to point at the input image data. 
+	// The format for FColor is B8G8R8A8, so the corresponding OpenCV format is CV_8UC4 (8-bit per channel, 4 channels)
+	const cv::Mat ImageMat = cv::Mat(ImageSize.Y, ImageSize.X, CV_8UC4, Image.GetData());
+
+	// Convert the image to grayscale before attempting to detect checkerboard corners
+	cv::Mat GrayImage;
+	cv::cvtColor(ImageMat, GrayImage, cv::COLOR_RGBA2GRAY);
+
+	const cv::Size CheckerboardSize = cv::Size(CheckerboardDimensions.X, CheckerboardDimensions.Y);
+	const int32 NumExpectedCorners = CheckerboardDimensions.X * CheckerboardDimensions.Y;
+
+	// ADAPTIVE_THRESH and NORMALIZE_IMAGE are both default options to aid in corner detection. 
+	// FAST_CHECK will quickly determine if there is any checkerboard in the image and early-out quickly if not
+	const int32 FindCornersFlags = cv::CALIB_CB_ADAPTIVE_THRESH | cv::CALIB_CB_NORMALIZE_IMAGE | cv::CALIB_CB_FAST_CHECK;
+
+	// Initialize the output array to hold the number of expected corners, then create a matrix header to the data to pass to OpenCV
+	OutCorners.Init(FVector2f(), NumExpectedCorners);
+	cv::Mat CornerMat = cv::Mat(NumExpectedCorners, 1, CV_32FC2, (void*)OutCorners.GetData());
+
+	const bool bFoundCorners = cv::findChessboardCorners(GrayImage, CheckerboardSize, CornerMat, FindCornersFlags);
+
+	if (!bFoundCorners || OutCorners.Num() != NumExpectedCorners)
+	{
+		OutCorners.Empty();
+		return false;
+	}
+
+	// Attempt to refine the corner detection to subpixel accuracy. Algorithm settings are all defaults from the OpenCV documentation.
+	const cv::TermCriteria Criteria = cv::TermCriteria(cv::TermCriteria::Type::EPS | cv::TermCriteria::Type::COUNT, 30, 0.001);
+	const cv::Size HalfSearchWindowSize = cv::Size(11, 11);
+	const cv::Size HalfDeadZoneSize = cv::Size(-1, -1);
+	cv::cornerSubPix(GrayImage, CornerMat, HalfSearchWindowSize, HalfDeadZoneSize, Criteria);
+
+	// The detected corner array can begin with either the TopLeft or BottomRight corner. If the BottomRight corner is first, reverse the order.
+	if ((NumExpectedCorners >= 2) && (OutCorners[0].Y > OutCorners[NumExpectedCorners - 1].Y))
+	{
+		Algo::Reverse(OutCorners);
+	}
+
+	return true;
+#else
+	return false;
+#endif // WITH_OPENCV
+}
+
+bool FOpenCVHelper::DrawCheckerboardCorners(const TArray<FVector2f>& Corners, FIntPoint CheckerboardDimensions, UTexture2D* DebugTexture)
+{
+#if WITH_OPENCV
+	if (!DebugTexture || Corners.IsEmpty())
+	{
+		return false;
+	}
+
+	cv::Mat CornerMat(Corners.Num(), 1, CV_32FC2, (void*)Corners.GetData());
+
+	FTexture2DMipMap& Mip0 = DebugTexture->GetPlatformData()->Mips[0];
+	void* TextureData = Mip0.BulkData.Lock(LOCK_READ_WRITE);
+
+	// Create a matrix header pointing to the raw texture data, draw the debug pattern, and update the texture resource
+	const FIntPoint TextureSize = FIntPoint(DebugTexture->GetSizeX(), DebugTexture->GetSizeY());
+	cv::Mat TextureMat = cv::Mat(TextureSize.Y, TextureSize.X, CV_8UC4, TextureData);
+
+	constexpr bool bFoundPattern = true;
+	cv::drawChessboardCorners(TextureMat, cv::Size(CheckerboardDimensions.X, CheckerboardDimensions.Y), CornerMat, bFoundPattern);
+
+	Mip0.BulkData.Unlock();
+
+	DebugTexture->UpdateResource();
+
+	return true;
+#else
+	return false;
+#endif // WITH_OPENCV
+}
+
+bool FOpenCVHelper::SolvePnP(const TArray<FVector>& ObjectPoints, const TArray<FVector2f>& ImagePoints, const FVector2D& FocalLength, const FVector2D& ImageCenter, const TArray<float>& DistortionParameters, FTransform& OutCameraPose)
+{
+#if WITH_OPENCV
+	const int32 NumPoints = ObjectPoints.Num();
+	if (!ensureMsgf((ImagePoints.Num() == NumPoints), TEXT("The number of 3D object points must match the number of 2D image points being passed to solvePnP()")))
+	{
+		return false;
+	}
+
+	// cv::solvePnP() will only accept spherical distortion parameters, but it accepts a variable number of parameters based on how much of the distortion model is used
+	// We need to guard against an incorrect number of parameters to avoid crashing in the opencv module
+	const int32 NumDistortionParameters = DistortionParameters.Num();
+	const bool bCorrectNumDistortionParameters = (NumDistortionParameters == 0) || (NumDistortionParameters == 4)  || (NumDistortionParameters == 5) || 
+		                                         (NumDistortionParameters == 8) || (NumDistortionParameters == 12) || (NumDistortionParameters == 14);
+	
+	if (!ensureMsgf(bCorrectNumDistortionParameters, TEXT("The number of distortion parameters being passed to solvePnP() is invalid")))
+	{
+		return false;
+	}
+
+	// Convert from UE coordinates to OpenCV coordinates
+	TArray<FVector> CvObjectPoints;
+	CvObjectPoints.Reserve(NumPoints);
+
+	for (const FVector& ObjectPoint : ObjectPoints)
+	{
+		CvObjectPoints.Add(FOpenCVHelper::ConvertUnrealToOpenCV(ObjectPoint));
+	}
+
+	cv::Mat ObjectPointsMat(NumPoints, 1, CV_64FC3, (void*)CvObjectPoints.GetData());
+	cv::Mat ImagePointsMat(NumPoints, 1, CV_32FC2, (void*)ImagePoints.GetData());
+	cv::Mat DistortionMat(NumDistortionParameters, 1, CV_32FC1, (void*)DistortionParameters.GetData());
+
+	// Initialize the camera matrix that will be used in each call to projectPoints()
+	cv::Mat CameraMatrix = cv::Mat::eye(3, 3, CV_64F);
+
+	CameraMatrix.at<double>(0, 0) = FocalLength.X;
+	CameraMatrix.at<double>(1, 1) = FocalLength.Y;
+	CameraMatrix.at<double>(0, 2) = ImageCenter.X;
+	CameraMatrix.at<double>(1, 2) = ImageCenter.Y;
+
+	// Solve for camera position
+	cv::Mat Rotation;
+	cv::Mat Translation;
+
+	// We send no distortion parameters, because Points2d was manually undistorted already
+	if (!cv::solvePnP(ObjectPointsMat, ImagePointsMat, CameraMatrix, DistortionMat, Rotation, Translation))
+	{
+		return false;
+	}
+
+	// Convert the OpenCV rotation and translation vectors into an FTransform in UE's coordinate system
+	MakeCameraPoseFromObjectVectors(Rotation, Translation, OutCameraPose);
+
+	return true;
+#else
+	return false;
+#endif // WITH_OPENCV
+}
+
+bool FOpenCVHelper::FitLine3D(const TArray<FVector>& InPoints, FVector& OutLine, FVector& OutPointOnLine)
+{
+#if WITH_OPENCV
+	const int32 NumPoints = InPoints.Num();
+
+	if (NumPoints < 2)
+	{
+		return false;
+	}
+
+	cv::Mat PointsMat(NumPoints, 1, CV_64FC3, (void*)InPoints.GetData());
+
+	// Find a best fit line between the 3D points, producing a line (the optical axis) and a point on that line
+	constexpr int CParam = 0; // fitline parameter of 0 will try to find an optimal value for the simple euclidean distance method of DIST_L2
+	constexpr double RadiusAccuracy = 0.01; // cv documentation specifies 0.01 as a good default value
+	constexpr double AngleAccuracy = 0.01; // cv documentation specifies 0.01 as a good default value
+
+	cv::Vec6f LineAndPoint;
+	cv::fitLine(PointsMat, LineAndPoint, cv::DIST_L2, CParam, RadiusAccuracy, AngleAccuracy);
+
+	OutLine = FVector(LineAndPoint[0], LineAndPoint[1], LineAndPoint[2]);
+	OutPointOnLine = FVector(LineAndPoint[3], LineAndPoint[4], LineAndPoint[5]);
+
+	return true;
+#else
+	return false;
+#endif // WITH_OPENCV
+}
+
+#if WITH_OPENCV
 double FOpenCVHelper::ComputeReprojectionError(const FTransform& CameraPose, const cv::Mat& CameraIntrinsicMatrix, const std::vector<cv::Point3f>& Points3d, const std::vector<cv::Point2f>& Points2d)
 {
 	// Ensure that the number of point correspondences is valid
@@ -312,66 +700,101 @@
 
 	return ReprojectionError;
 }
+#endif
 
-void FOpenCVHelper::ConvertTransformToVectors(const FTransform& InTransform, cv::Mat& OutRotation, cv::Mat& OutTranslation)
+double FOpenCVHelper::ComputeReprojectionError(const TArray<FVector>& ObjectPoints, const TArray<FVector2f>& ImagePoints, const FVector2D& FocalLength, const FVector2D& ImageCenter, const FTransform& CameraPose)
 {
-	const FMatrix Matrix = InTransform.ToMatrixNoScale();
-
-	const cv::Mat Tcam = (cv::Mat_<double>(3, 1) << Matrix.M[3][0], Matrix.M[3][1], Matrix.M[3][2]);
-
-	cv::Mat Rcam = cv::Mat::zeros(3, 3, cv::DataType<double>::type);
-	for (int32 Column = 0; Column < 3; ++Column)
+#if WITH_OPENCV
+	// Ensure that the number of point correspondences is valid
+	const int32 NumPoints3d = ObjectPoints.Num();
+	const int32 NumPoints2d = ImagePoints.Num();
+	if ((NumPoints3d == 0) || (NumPoints2d == 0) || (NumPoints3d != NumPoints2d))
 	{
-		FVector ColVec = Matrix.GetColumn(Column);
-		Rcam.at<double>(Column, 0) = ColVec.X;
-		Rcam.at<double>(Column, 1) = ColVec.Y;
-		Rcam.at<double>(Column, 2) = ColVec.Z;
+		return -1.0;
 	}
 
-	const cv::Mat Robj = Rcam.t();
+	// Initialize the camera matrix that will be used in each call to projectPoints()
+	cv::Mat CameraMatrix = cv::Mat::eye(3, 3, CV_64F);
 
-	cv::Rodrigues(Robj, OutRotation);
+	CameraMatrix.at<double>(0, 0) = FocalLength.X;
+	CameraMatrix.at<double>(1, 1) = FocalLength.Y;
+	CameraMatrix.at<double>(0, 2) = ImageCenter.X;
+	CameraMatrix.at<double>(1, 2) = ImageCenter.Y;
 
-	OutTranslation = -Rcam.inv() * Tcam;
+	cv::Mat Rotation;
+	cv::Mat Translation;
+	FOpenCVHelper::MakeObjectVectorsFromCameraPose(CameraPose, Rotation, Translation);
+
+	// cv::projectPoints requires that the 3D points and 2D points have the same bit depth, so we have to make a copy of the ObjectPoints with only 32-bits of precision
+	TArray<FVector3f> CvObjectPoints;
+	for (const FVector& Point : ObjectPoints)
+	{
+		FVector CvPoint = ConvertUnrealToOpenCV(Point);
+		CvObjectPoints.Add(FVector3f(CvPoint.X, CvPoint.Y, CvPoint.Z));
+	}
+
+	TArray<FVector2f> ProjectedPoints;
+	ProjectedPoints.Init(FVector2f(), NumPoints2d);
+
+	cv::Mat ObjectPointsMat = cv::Mat(ObjectPoints.Num(), 1, CV_32FC3, (void*)CvObjectPoints.GetData());
+	cv::Mat ProjectedPointsMat = cv::Mat(ProjectedPoints.Num(), 1, CV_32FC2, (void*)ProjectedPoints.GetData());
+
+	// The 2D points will be compared against the undistorted 2D points, so the distortion coefficients can be ignored
+	cv::projectPoints(ObjectPointsMat, Rotation, Translation, CameraMatrix, cv::noArray(), ProjectedPointsMat);
+
+	if (ProjectedPoints.Num() != NumPoints2d)
+	{
+		return -1.0;
+	}
+
+	// Compute euclidean distance between captured 2D points and reprojected 2D points to measure reprojection error
+	double ReprojectionError = 0.0;
+	for (int32 Index = 0; Index < NumPoints2d; ++Index)
+	{
+		const FVector2f Diff = ProjectedPoints[Index] - ImagePoints[Index];
+		ReprojectionError += (Diff.X * Diff.X) + (Diff.Y * Diff.Y); // cv::norm with NORM_L2SQR
+	}
+
+	return ReprojectionError;
+#else 
+	return -1.0f;
+#endif // WITH_OPENCV
 }
 
-void FOpenCVHelper::ConvertVectorsToTransform(const cv::Mat& InRotation, const cv::Mat& InTranslation, FTransform& OutTransform)
+#if WITH_OPENCV
+cv::Mat FOpenCVLensDistortionParametersBase::ConvertToOpenCVDistortionCoefficients() const
 {
-	// Convert to camera pose
-
-	// [R|t]' = [R'|-R'*t]
-
-	// Convert from Rodrigues to rotation matrix
-	cv::Mat Robj;
-	cv::Rodrigues(InRotation, Robj); // Robj is 3x3
-
-	// Calculate camera translation
-	cv::Mat Tcam = -Robj.t() * InTranslation;
-
-	// Invert/transpose to get camera orientation
-	cv::Mat Rcam = Robj.t();
-
-	// Convert back to UE coordinates
-
-	FMatrix M = FMatrix::Identity;
-
-	// Fill rotation matrix
-	for (int32 Column = 0; Column < 3; ++Column)
+	if (bUseFisheyeModel)
 	{
-		M.SetColumn(Column, FVector(
-			Rcam.at<double>(Column, 0),
-			Rcam.at<double>(Column, 1),
-			Rcam.at<double>(Column, 2))
-		);
+		cv::Mat DistortionCoefficients(1, 4, CV_64F);
+		DistortionCoefficients.at<double>(0) = K1;
+		DistortionCoefficients.at<double>(1) = K2;
+		DistortionCoefficients.at<double>(2) = K3;
+		DistortionCoefficients.at<double>(3) = K4;
+		return DistortionCoefficients;
 	}
-
-	// Fill translation vector
-	M.M[3][0] = Tcam.at<double>(0);
-	M.M[3][1] = Tcam.at<double>(1);
-	M.M[3][2] = Tcam.at<double>(2);
-
-	OutTransform.SetFromMatrix(M);
-	FOpenCVHelper::ConvertOpenCVToUnreal(OutTransform);
+	else
+	{
+		cv::Mat DistortionCoefficients(1, 8, CV_64F);
+		DistortionCoefficients.at<double>(0) = K1;
+		DistortionCoefficients.at<double>(1) = K2;
+		DistortionCoefficients.at<double>(2) = P1;
+		DistortionCoefficients.at<double>(3) = P2;
+		DistortionCoefficients.at<double>(4) = K3;
+		DistortionCoefficients.at<double>(5) = K4;
+		DistortionCoefficients.at<double>(6) = K5;
+		DistortionCoefficients.at<double>(7) = K6;
+		return DistortionCoefficients;
+	}
 }
 
+cv::Mat FOpenCVLensDistortionParametersBase::CreateOpenCVCameraMatrix(const FVector2D& InImageSize) const
+{
+	cv::Mat CameraMatrix = cv::Mat::eye(3, 3, CV_64F);
+	CameraMatrix.at<double>(0, 0) = F.X * InImageSize.X;
+	CameraMatrix.at<double>(1, 1) = F.Y * InImageSize.Y;
+	CameraMatrix.at<double>(0, 2) = C.X * InImageSize.X;
+	CameraMatrix.at<double>(1, 2) = C.Y * InImageSize.Y;
+	return CameraMatrix;
+}
 #endif // WITH_OPENCV
diff --git a/Engine/Plugins/Runtime/OpenCV/Source/OpenCVHelper/Public/OpenCVHelper.h b/Engine/Plugins/Runtime/OpenCV/Source/OpenCVHelper/Public/OpenCVHelper.h
--- a/Engine/Plugins/Runtime/OpenCV/Source/OpenCVHelper/Public/OpenCVHelper.h
+++ b/Engine/Plugins/Runtime/OpenCV/Source/OpenCVHelper/Public/OpenCVHelper.h
@@ -26,12 +26,12 @@
 
 #define OPENCV_INCLUDES_END static_assert(false, "Include PostOpenCVHeaders.h instead of using this macro");
 
-#if WITH_OPENCV
-
 class UTexture2D;
 class FString;
 class FName;
 
+#if WITH_OPENCV
+
 namespace cv
 {
 	class Mat;
@@ -41,8 +41,37 @@ namespace cv
 	typedef Point_<float> Point2f;
 	typedef Point3_<float> Point3f;
 };
+
 #endif	// WITH_OPENCV
 
+UENUM(BlueprintType)
+enum class EArucoDictionary : uint8
+{
+	None                 UMETA(DisplayName = "None"),
+	DICT_4X4_50          UMETA(DisplayName = "DICT_4X4_50"),
+	DICT_4X4_100		 UMETA(DisplayName = "DICT_4X4_100"),
+	DICT_4X4_250		 UMETA(DisplayName = "DICT_4X4_250"),
+	DICT_4X4_1000		 UMETA(DisplayName = "DICT_4X4_1000"),
+	DICT_5X5_50			 UMETA(DisplayName = "DICT_5X5_50"),
+	DICT_5X5_100		 UMETA(DisplayName = "DICT_5X5_100"),
+	DICT_5X5_250		 UMETA(DisplayName = "DICT_5X5_250"),
+	DICT_5X5_1000		 UMETA(DisplayName = "DICT_5X5_1000"),
+	DICT_6X6_50			 UMETA(DisplayName = "DICT_6X6_50"),
+	DICT_6X6_100		 UMETA(DisplayName = "DICT_6X6_100"),
+	DICT_6X6_250		 UMETA(DisplayName = "DICT_6X6_250"),
+	DICT_6X6_1000		 UMETA(DisplayName = "DICT_6X6_1000"),
+	DICT_7X7_50			 UMETA(DisplayName = "DICT_7X7_50"),
+	DICT_7X7_100		 UMETA(DisplayName = "DICT_7X7_100"),
+	DICT_7X7_250		 UMETA(DisplayName = "DICT_7X7_250"),
+	DICT_7X7_1000		 UMETA(DisplayName = "DICT_7X7_1000"),
+	DICT_ARUCO_ORIGINAL	 UMETA(DisplayName = "DICT_ARUCO_ORIGINAL")
+};
+
+struct OPENCVHELPER_API FArucoMarker
+{
+	int32 MarkedId;
+	FVector2f Corners[4];
+};
 
 class OPENCVHELPER_API FOpenCVHelper
 {
@@ -70,7 +99,7 @@ public:
 		return UnitVectors[std::underlying_type_t<EAxis>(Axis)];
 	};
 
-	/** Converts in-place the coordinate system of the given FTransform by specifying the source axes in terms of the destionation axes */
+	/** Converts in-place the coordinate system of the given FTransform by specifying the source axes in terms of the destination axes */
 	static void ConvertCoordinateSystem(FTransform& Transform, const EAxis DstXInSrcAxis, const EAxis DstYInSrcAxis, const EAxis DstZInSrcAxis);
 
 	/** Converts in-place an FTransform in Unreal coordinates to OpenCV coordinates */
@@ -79,8 +108,14 @@ public:
 	/** Converts in-place an FTransform in OpenCV coordinates to Unreal coordinates */
 	static void ConvertOpenCVToUnreal(FTransform& Transform);
 
-#if WITH_OPENCV
+	/** Converts an FVector in Unreal coordinates to OpenCV coordinates */
+	static FVector ConvertUnrealToOpenCV(const FVector& Transform);
+
+	/** Converts an FTransform in OpenCV coordinates to Unreal coordinates */
+	static FVector ConvertOpenCVToUnreal(const FVector& Transform);
+
 public:
+#if WITH_OPENCV
 	/**
 	 * Creates a Texture from the given Mat, if its properties (e.g. pixel format) are supported.
 	 * 
@@ -93,14 +128,46 @@ public:
 	static UTexture2D* TextureFromCvMat(cv::Mat& Mat, const FString* PackagePath = nullptr, const FName* TextureName = nullptr);
 	static UTexture2D* TextureFromCvMat(cv::Mat& Mat, UTexture2D* InTexture);
 
-	/** Converts an FTransform to a rotation vector (cv::Mat) and a translation vector (cv::Mat) */
-	static void ConvertTransformToVectors(const FTransform& InTransform, cv::Mat& OutRotation, cv::Mat& OutTranslation);
+	/** 
+	 * Takes a rotation vector (in rodrigues form) and translation vector that represent a change of basis from object space to camera space (in OpenCV's coordinate system),
+	 * inverts them to generate a camera pose, and converts it to an FTransform in UE's coordinate system. 
+	 * The input rotation and translation vectors are expected to match the form of an "rvec" and "tvec" that are generated functions like cv::calibrateCamera() and cv::solvePnP()
+	 */
+	static void MakeCameraPoseFromObjectVectors(const cv::Mat& InRotation, const cv::Mat& InTranslation, FTransform& OutTransform);
+
+	/**
+	 * Takes a camera pose transform in UE's coordinate system, converts it to a rotation vector (in rodrigues form) and translation vector (in OpenCV's coordinate system),
+	 * and inverts them to generate vectors that represent a change of basis from object space to camera space.
+	 * The output rotation and translation vectors are expected to match the form of an "rvec" and "tvec" that are consumed by functions like cv::calibrateCamera() and cv::solvePnP()
+	 */
+	static void MakeObjectVectorsFromCameraPose(const FTransform& InTransform, cv::Mat& OutRotation, cv::Mat& OutTranslation);
+#endif	// WITH_OPENCV
+
+	/** Identify a set of aruco markers in the input image that belong to the input aruco dictionary, and output the marker IDs and the 2D coordinates of the 4 corners of each marker */
+	static bool IdentifyArucoMarkers(TArray<FColor>& Image, FIntPoint ImageSize, EArucoDictionary DictionaryName, TArray<FArucoMarker>& OutMarkers);
+
+	/** Draw a debug view of the input aruco markers on top of the input texture */
+	static bool DrawArucoMarkers(const TArray<FArucoMarker>& Markers, UTexture2D* DebugTexture);
 
-	/** Converts a rotation vector (cv::Mat) and a translation vector (cv::Mat) to an FTransform */
-	static void ConvertVectorsToTransform(const cv::Mat& InRotation, const cv::Mat& InTranslation, FTransform& OutTransform);
+	/** Identify a checkerboard pattern in the input image that with the given checkerboard dimensions (columns x rows), and output the 2D coordinates of the intersections between each checkerboard square */
+	static bool IdentifyCheckerboard(TArray<FColor>& Image, FIntPoint ImageSize, FIntPoint CheckerboardDimensions, TArray<FVector2f>& OutCorners);
 
+	/** Draw a debug view of the input checkerboard corners on top of the input texture */
+	static bool DrawCheckerboardCorners(const TArray<FVector2f>& Corners, FIntPoint CheckerboardDimensions, UTexture2D* DebugTexture);
+
+	/** Compute the camera pose that minimizes the reprojection error of the input object points and image points */
+	static bool SolvePnP(const TArray<FVector>& ObjectPoints, const TArray<FVector2f>& ImagePoints, const FVector2D& FocalLength, const FVector2D& ImageCenter, const TArray<float>& DistortionParameters, FTransform& OutCameraPose);
+
+	/** Find a 3D fit line that passes through the input points, as well as a point on that line */
+	static bool FitLine3D(const TArray<FVector>& InPoints, FVector& OutLine, FVector& OutPointOnLine);
+
+#if WITH_OPENCV
+	UE_DEPRECATED(5.4, "The version of ComputeReprojectionError takes OpenCV types as input parameters is deprecated. Please use the version that takes all UE types.")
 	static double ComputeReprojectionError(const FTransform& CameraPose, const cv::Mat& CameraIntrinsicMatrix, const std::vector<cv::Point3f>& Points3d, const std::vector<cv::Point2f>& Points2d);
-#endif	// WITH_OPENCV
+#endif
+
+	/** Project the 3D objects points to the 2D image plane represented by the input camera intrinsics and pose, and compute the reprojection error (euclidean distance) between the input image points and the reprojected points */
+	static double ComputeReprojectionError(const TArray<FVector>& ObjectPoints, const TArray<FVector2f>& ImagePoints, const FVector2D& FocalLength, const FVector2D& ImageCenter, const FTransform& CameraPose);
 };
 
 /**
diff --git a/Engine/Plugins/VirtualProduction/CameraCalibration/Source/CameraCalibrationEditor/Private/AssetEditor/CameraCalibrationStepsController.cpp b/Engine/Plugins/VirtualProduction/CameraCalibration/Source/CameraCalibrationEditor/Private/AssetEditor/CameraCalibrationStepsController.cpp
--- a/Engine/Plugins/VirtualProduction/CameraCalibration/Source/CameraCalibrationEditor/Private/AssetEditor/CameraCalibrationStepsController.cpp
+++ b/Engine/Plugins/VirtualProduction/CameraCalibration/Source/CameraCalibrationEditor/Private/AssetEditor/CameraCalibrationStepsController.cpp
@@ -50,11 +50,6 @@
 #include "Widgets/SWidget.h"
 
 
-#if WITH_OPENCV
-#include "OpenCVHelper.h"
-#endif
-
-
 #define LOCTEXT_NAMESPACE "CameraCalibrationStepsController"
 
 namespace CameraCalibrationStepsController
@@ -1638,7 +1633,7 @@ UTextureRenderTarget2D* FCameraCalibrationStepsController::GetMediaPlateRenderTa
 	return MediaPlateRenderTarget.Get();
 }
 
-bool FCameraCalibrationStepsController::CalculateNormalizedMouseClickPosition(const FGeometry& MyGeometry, const FPointerEvent& MouseEvent, FVector2D& OutPosition, ESimulcamViewportPortion ViewportPortion) const
+bool FCameraCalibrationStepsController::CalculateNormalizedMouseClickPosition(const FGeometry& MyGeometry, const FPointerEvent& MouseEvent, FVector2f& OutPosition, ESimulcamViewportPortion ViewportPortion) const
 {
 	// Reject viewports with no area
 	if (FMath::IsNearlyZero(MyGeometry.Size.X) || FMath::IsNearlyZero(MyGeometry.Size.Y))
@@ -1653,7 +1648,7 @@ bool FCameraCalibrationStepsController::CalculateNormalizedMouseClickPosition(co
 	// * MyGeometry.AbsolutePosition        : Position of the top-left corner of viewport within screen
 	// * MyGeometry.AbsoluteToLocal(MouseEvent.GetScreenSpacePosition()) gives you the pixel coordinates local to the viewport.
 
-	const FVector2D LocalInPixels = MyGeometry.AbsoluteToLocal(MouseEvent.GetScreenSpacePosition());
+	const FVector2f LocalInPixels = MyGeometry.AbsoluteToLocal(MouseEvent.GetScreenSpacePosition());
 
 	float XNormalized = LocalInPixels.X / MyGeometry.Size.X;
 	float YNormalized = LocalInPixels.Y / MyGeometry.Size.Y;
@@ -1683,7 +1678,7 @@ bool FCameraCalibrationStepsController::CalculateNormalizedMouseClickPosition(co
 	}
 
 	// Position 0~1. Origin at top-left corner of the viewport.
-	OutPosition = FVector2D(XNormalized, YNormalized);
+	OutPosition = FVector2f(XNormalized, YNormalized);
 
 	return true;
 }
diff --git a/Engine/Plugins/VirtualProduction/CameraCalibration/Source/CameraCalibrationEditor/Private/AssetEditor/CameraCalibrationStepsController.h b/Engine/Plugins/VirtualProduction/CameraCalibration/Source/CameraCalibrationEditor/Private/AssetEditor/CameraCalibrationStepsController.h
--- a/Engine/Plugins/VirtualProduction/CameraCalibration/Source/CameraCalibrationEditor/Private/AssetEditor/CameraCalibrationStepsController.h
+++ b/Engine/Plugins/VirtualProduction/CameraCalibration/Source/CameraCalibrationEditor/Private/AssetEditor/CameraCalibrationStepsController.h
@@ -130,7 +130,7 @@ public:
 	void SelectStep(const FName& Name);
 
 	/** Calculates the normalized (0~1) coordinates in the simulcam viewport of the given mouse click */
-	bool CalculateNormalizedMouseClickPosition(const FGeometry& MyGeometry, const FPointerEvent& MouseEvent, FVector2D& OutPosition, ESimulcamViewportPortion ViewportPortion = ESimulcamViewportPortion::FullViewport) const;
+	bool CalculateNormalizedMouseClickPosition(const FGeometry& MyGeometry, const FPointerEvent& MouseEvent, FVector2f& OutPosition, ESimulcamViewportPortion ViewportPortion = ESimulcamViewportPortion::FullViewport) const;
 
 	/** Finds the world being used by the tool for finding and spawning objects */
 	UWorld* GetWorld() const;
diff --git a/Engine/Plugins/VirtualProduction/CameraCalibration/Source/CameraCalibrationEditor/Private/Calibrators/CameraCalibrationSolver.cpp b/Engine/Plugins/VirtualProduction/CameraCalibration/Source/CameraCalibrationEditor/Private/Calibrators/CameraCalibrationSolver.cpp
--- a/Engine/Plugins/VirtualProduction/CameraCalibration/Source/CameraCalibrationEditor/Private/Calibrators/CameraCalibrationSolver.cpp
+++ b/Engine/Plugins/VirtualProduction/CameraCalibration/Source/CameraCalibrationEditor/Private/Calibrators/CameraCalibrationSolver.cpp
@@ -16,23 +16,24 @@
 #endif	// WITH_OPENCV
 
 static TAutoConsoleVariable<int> CVarUseLegacySphericalSolver(TEXT("CameraCalibration.UseLegacySphericalSolver"), 0, TEXT("If set, the legacy OpenCV spherical solver will be used"));
+static TAutoConsoleVariable<float> CVarRotationStepValue(TEXT("CameraCalibration.RotationStepValue"), 0.05, TEXT("The value of the initial step size to use when finding an optimal nodal offset rotation that minimizes reprojection error."));
+static TAutoConsoleVariable<float> CVarLocationStepValue(TEXT("CameraCalibration.LocationStepValue"), 0.5, TEXT("The value of the initial step size to use when finding an optimal nodal offset location that minimizes reprojection error."));
 
 DEFINE_LOG_CATEGORY_STATIC(LogCameraCalibrationSolver, Log, All);
 
-#if WITH_OPENCV
-
 double FCameraCalibrationSolver::CalibrateCamera(
 	const TSubclassOf<ULensModel> LensModel,
 	const TArray<TArray<FVector>>& InObjectPoints,
-	const TArray<TArray<FVector2D>>& InImagePoints,
+	const TArray<TArray<FVector2f>>& InImagePoints,
 	const FIntPoint ImageSize,
-	FVector2f& InOutFocalLength,
-	FVector2f& InOutImageCenter,
+	FVector2D& InOutFocalLength,
+	FVector2D& InOutImageCenter,
 	TArray<float>& OutDistCoeffs,
 	TArray<FTransform>& InOutCameraPoses,
 	double PixelAspect,
 	ECalibrationFlags SolverFlags)
 {
+#if WITH_OPENCV
 	const int NumImages = InObjectPoints.Num();
 
 	// Create an array to store the number of points in each image
@@ -50,11 +51,29 @@ double FCameraCalibrationSolver::CalibrateCamera(
 		MaxPoints = MAX(MaxPoints, NumPointsInImage);
 	}
 
+	// Convert the object points from Unreal's coordinate system to OpenCV's
+	TArray<TArray<FVector>> CvObjectPoints;
+	CvObjectPoints.Reserve(NumImages);
+	for (int ImageIndex = 0; ImageIndex < NumImages; ++ImageIndex)
+	{
+		const TArray<FVector>& ObjectPointsForImage = InObjectPoints[ImageIndex];
+
+		TArray<FVector> CvObjectPointsForImage;
+		CvObjectPointsForImage.Reserve(ObjectPointsForImage.Num());
+
+		for (const FVector& Point : ObjectPointsForImage)
+		{
+			CvObjectPointsForImage.Add(FOpenCVHelper::ConvertUnrealToOpenCV(Point));
+		}
+
+		CvObjectPoints.Add(CvObjectPointsForImage);
+	}
+
 	cv::Mat ObjectPointsMat = cv::Mat(1, NumTotalPoints, CV_64FC3);
 	cv::Mat ImagePointsMat = cv::Mat(1, NumTotalPoints, CV_64FC2);
 
 	// Reorganize the 3D and 2D points from the input arrays or arrays to be laid out linearly in memory in two cv::Mat objects
-	GatherPoints(InObjectPoints, InImagePoints, ObjectPointsMat, ImagePointsMat);
+	GatherPoints(CvObjectPoints, InImagePoints, ObjectPointsMat, ImagePointsMat);
 
 	double RMSE = 0.0;
 
@@ -96,9 +115,9 @@ double FCameraCalibrationSolver::CalibrateCamera(
 		std::vector<std::vector<cv::Point3f>> Samples3d;
 
 		Samples2d.reserve(InImagePoints.Num());
-		Samples3d.reserve(InObjectPoints.Num());
+		Samples3d.reserve(CvObjectPoints.Num());
 
-		for (const TArray<FVector>& Image : InObjectPoints)
+		for (const TArray<FVector>& Image : CvObjectPoints)
 		{
 			std::vector<cv::Point3f> Points3d;
 			Points3d.reserve(Image.Num());
@@ -111,12 +130,12 @@ double FCameraCalibrationSolver::CalibrateCamera(
 			Samples3d.push_back(Points3d);
 		}
 
-		for (const TArray<FVector2D>& Image : InImagePoints)
+		for (const TArray<FVector2f>& Image : InImagePoints)
 		{
 			std::vector<cv::Point2f> Points2d;
 			Points2d.reserve(Image.Num());
 
-			for (const FVector2D& Point2d : Image)
+			for (const FVector2f& Point2d : Image)
 			{
 				Points2d.push_back(cv::Point2f(Point2d.X, Point2d.Y));
 			}
@@ -248,7 +267,7 @@ double FCameraCalibrationSolver::CalibrateCamera(
 		}
 		else
 		{
-			FOpenCVHelper::ConvertTransformToVectors(InOutCameraPoses[ImageIndex], Rotation, Translation);
+			FOpenCVHelper::MakeObjectVectorsFromCameraPose(InOutCameraPoses[ImageIndex], Rotation, Translation);
 		}
 	}
 
@@ -367,12 +386,92 @@ double FCameraCalibrationSolver::CalibrateCamera(
 		cv::Mat Rotation = Solver.Params.rowRange(ExtrinsicOffset, ExtrinsicOffset + 3);
 		cv::Mat Translation = Solver.Params.rowRange(ExtrinsicOffset + 3, ExtrinsicOffset + 6);
 
-		FOpenCVHelper::ConvertVectorsToTransform(Rotation, Translation, InOutCameraPoses[ImageIndex]);
+		FOpenCVHelper::MakeCameraPoseFromObjectVectors(Rotation, Translation, InOutCameraPoses[ImageIndex]);
 	}
 
 	return RMSE;
+#else
+	return -1.0;
+#endif // WITH_OPENCV
 }
 
+double FCameraCalibrationSolver::OptimizeNodalOffset(
+	const TArray<TArray<FVector>>& InObjectPoints,
+	const TArray<TArray<FVector2f>>& InImagePoints,
+	const FVector2D& InFocalLength,
+	const FVector2D& InImageCenter,
+	const TArray<FTransform>& InCameraPoses,
+	FTransform& InOutNodalOffset)
+{
+#if WITH_OPENCV
+	const int32 NumViews = InObjectPoints.Num();
+	if (NumViews < 1)
+	{
+		return -1.0;
+	}
+
+	cv::Ptr<cv::DownhillSolver> Solver = cv::DownhillSolver::create();
+
+	const FQuat InitialRotation = InOutNodalOffset.GetRotation();
+	const FVector InitialLocation = InOutNodalOffset.GetLocation();
+
+	constexpr int32 NumRotationParameters = 4; // FQuat
+	constexpr int32 NumLocationParameters = 3; // FVector
+	cv::Mat WorkingSolution = cv::Mat(1, NumRotationParameters + NumLocationParameters, CV_64FC1);
+	WorkingSolution.at<double>(0, 0) = InitialRotation.X;
+	WorkingSolution.at<double>(0, 1) = InitialRotation.Y;
+	WorkingSolution.at<double>(0, 2) = InitialRotation.Z;
+	WorkingSolution.at<double>(0, 3) = InitialRotation.W;
+	WorkingSolution.at<double>(0, 4) = InitialLocation.X;
+	WorkingSolution.at<double>(0, 5) = InitialLocation.Y;
+	WorkingSolution.at<double>(0, 6) = InitialLocation.Z;
+
+	// NOTE: These step sizes may need further testing and refinement, but tests so far have shown them to be decent
+	const double RotationStep = CVarRotationStepValue.GetValueOnGameThread();
+	const double LocationStep = CVarLocationStepValue.GetValueOnGameThread();
+	cv::Mat Step = cv::Mat(1, NumRotationParameters + NumLocationParameters, CV_64FC1);
+	Step.at<double>(0, 0) = RotationStep;
+	Step.at<double>(0, 1) = RotationStep;
+	Step.at<double>(0, 2) = RotationStep;
+	Step.at<double>(0, 3) = RotationStep;
+	Step.at<double>(0, 4) = LocationStep;
+	Step.at<double>(0, 5) = LocationStep;
+	Step.at<double>(0, 6) = LocationStep;
+
+	Solver->setInitStep(Step);
+
+	cv::Ptr<FOptimizeNodalOffsetSolver> SolverFunction = cv::makePtr<FOptimizeNodalOffsetSolver>();
+	Solver->setFunction(SolverFunction);
+
+	SolverFunction->FocalLength = InFocalLength;
+	SolverFunction->ImageCenter = InImageCenter;
+
+	SolverFunction->CameraPoses.Reserve(NumViews);
+	SolverFunction->Points3d.Reserve(NumViews);
+	SolverFunction->Points2d.Reserve(NumViews);
+
+	for (int32 ViewIndex = 0; ViewIndex < NumViews; ++ViewIndex)
+	{
+		SolverFunction->CameraPoses.Add(InCameraPoses[ViewIndex]);
+		SolverFunction->Points3d.Add(InObjectPoints[ViewIndex]);
+		SolverFunction->Points2d.Add(InImagePoints[ViewIndex]);
+	}
+
+	const double Error = Solver->minimize(WorkingSolution);
+
+	const FQuat FinalRotation = FQuat(WorkingSolution.at<double>(0, 0), WorkingSolution.at<double>(0, 1), WorkingSolution.at<double>(0, 2), WorkingSolution.at<double>(0, 3)).GetNormalized();
+	const FVector FinalLocation = FVector(WorkingSolution.at<double>(0, 4), WorkingSolution.at<double>(0, 5), WorkingSolution.at<double>(0, 6));
+
+	InOutNodalOffset.SetRotation(FinalRotation);
+	InOutNodalOffset.SetLocation(FinalLocation);
+
+	return Error;
+#else
+	return -1.0;
+#endif // WITH_OPENCV
+}
+
+#if WITH_OPENCV
 void FCameraCalibrationSolver::InitCameraIntrinsics(
 	const cv::Mat& ObjectPoints,
 	const cv::Mat& ImagePoints,
@@ -1226,7 +1325,7 @@ void FCameraCalibrationSolver::ProjectPointsSpherical(
 
 void FCameraCalibrationSolver::GatherPoints(
 	const TArray<TArray<FVector>>& ObjectPoints,
-	const TArray<TArray<FVector2D>>& ImagePoints,
+	const TArray<TArray<FVector2f>>& ImagePoints,
 	cv::Mat& ObjectPointsMat,
 	cv::Mat& ImagePointsMat)
 {
@@ -1240,7 +1339,7 @@ void FCameraCalibrationSolver::GatherPoints(
 	for (int ImageIndex = 0; ImageIndex < NumImages; ++ImageIndex)
 	{
 		const TArray<FVector>& CurrentObjectPoints = ObjectPoints[ImageIndex];
-		const TArray<FVector2D>& CurrentImagePoints = ImagePoints[ImageIndex];
+		const TArray<FVector2f>& CurrentImagePoints = ImagePoints[ImageIndex];
 
 		const int NumPointsInImage = CurrentObjectPoints.Num();
 
@@ -1478,4 +1577,52 @@ void FLevMarqSolver::Step()
 	}
 }
 
+int FOptimizeNodalOffsetSolver::getDims() const
+{
+	constexpr int32 NumRotationParameters = 4; // FQuat
+	constexpr int32 NumLocationParameters = 3; // FVector
+	return NumRotationParameters + NumLocationParameters;
+}
+
+double FOptimizeNodalOffsetSolver::calc(const double* x) const
+{
+	// Convert the input data (7 doubles) to an FQuat and FVector
+	const FQuat Rotation = FQuat(x[0], x[1], x[2], x[3]).GetNormalized();
+	const FVector Location = FVector(x[4], x[5], x[6]);
+
+	UE_LOG(LogCameraCalibrationSolver, VeryVerbose, TEXT("Nodal Offset Candidate:  Rotation: (%lf, %lf, %lf, %lf)  Location: (%lf, %lf, %lf)"),
+		Rotation.X, Rotation.Y, Rotation.Z, Rotation.W, Location.X, Location.Y, Location.Z);
+
+	FTransform NodalOffsetCandidate;
+
+	// As a result of the way that the downhill solver nudges the input data on each iteration, it is important to normalize the rotation
+	NodalOffsetCandidate.SetRotation(Rotation);
+	NodalOffsetCandidate.SetLocation(Location);
+
+	double ReprojectionErrorTotal = 0.0;
+	int32 NumTotalPoints = 0;
+
+	const int32 NumCameraViews = CameraPoses.Num();
+	for (int32 ViewIndex = 0; ViewIndex < NumCameraViews; ++ViewIndex)
+	{
+		const FTransform& CameraPose = CameraPoses[ViewIndex];
+		const TArray<FVector>& ObjectPoints = Points3d[ViewIndex];
+		const TArray<FVector2f>& ImagePoints = Points2d[ViewIndex];
+
+		// Compute the optimal camera pose using the nodal offset candidate for this iteration and the tracked camera pose 
+		const FTransform OptimalCameraPose = NodalOffsetCandidate * CameraPose;
+
+		// Compute the reprojection error for this view and add it to the running total
+		double ViewError = FOpenCVHelper::ComputeReprojectionError(ObjectPoints, ImagePoints, FocalLength, ImageCenter, OptimalCameraPose);
+		ReprojectionErrorTotal += ViewError;
+
+		NumTotalPoints += ImagePoints.Num();
+	}
+
+	const double RootMeanSquareError = FMath::Sqrt(ReprojectionErrorTotal / NumTotalPoints);
+
+	UE_LOG(LogCameraCalibrationSolver, VeryVerbose, TEXT("Reprojection Error: %lf"), RootMeanSquareError);
+
+	return RootMeanSquareError;
+}
 #endif // WITH_OPENCV
diff --git a/Engine/Plugins/VirtualProduction/CameraCalibration/Source/CameraCalibrationEditor/Private/Calibrators/CameraCalibrationSolver.h b/Engine/Plugins/VirtualProduction/CameraCalibration/Source/CameraCalibrationEditor/Private/Calibrators/CameraCalibrationSolver.h
--- a/Engine/Plugins/VirtualProduction/CameraCalibration/Source/CameraCalibrationEditor/Private/Calibrators/CameraCalibrationSolver.h
+++ b/Engine/Plugins/VirtualProduction/CameraCalibration/Source/CameraCalibrationEditor/Private/Calibrators/CameraCalibrationSolver.h
@@ -33,7 +33,6 @@ enum class ECalibrationFlags : uint32
   */
 class FCameraCalibrationSolver
 {
-#if WITH_OPENCV
 public:
 	/** 
 	  * Calibrate camera intrinsics and distortion parameters using a set of input 3D and 2D point correspondences
@@ -42,15 +41,25 @@ public:
 	static double CalibrateCamera(
 		const TSubclassOf<ULensModel> LensModel,
 		const TArray<TArray<FVector>>& InObjectPoints,
-		const TArray<TArray<FVector2D>>& InImagePoints,
+		const TArray<TArray<FVector2f>>& InImagePoints,
 		const FIntPoint ImageSize,
-		FVector2f& InOutFocalLength,
-		FVector2f& InOutImageCenter,
+		FVector2D& InOutFocalLength,
+		FVector2D& InOutImageCenter,
 		TArray<float>& OutDistCoeffs,
 		TArray<FTransform>& InOutCameraPoses,
 		double PixelAspect = 1.0,
 		ECalibrationFlags SolverFlags = ECalibrationFlags::None);
 
+	/** Optimize the input nodal offset transform by running a downhill solver that attempts to minimize the reprojection error of the input points */
+	static double OptimizeNodalOffset(
+		const TArray<TArray<FVector>>& InObjectPoints, 
+		const TArray<TArray<FVector2f>>& InImagePoints, 
+		const FVector2D& InFocalLength,
+		const FVector2D& InImageCenter,
+		const TArray<FTransform>& InCameraPoses,
+		FTransform& InOutNodalOffset);
+
+#if WITH_OPENCV
 private:
 	/** Initialize the camera matrix of intrinsic parameters using linear algebra techniques */
 	static void InitCameraIntrinsics(
@@ -122,7 +131,7 @@ private:
 	/** Copy the input 3D and 2D points to OpenCV matrices for ease of use with the solver */
 	static void GatherPoints(
 		const TArray<TArray<FVector>>& ObjectPoints,
-		const TArray<TArray<FVector2D>>& ImagePoints,
+		const TArray<TArray<FVector2f>>& ImagePoints,
 		cv::Mat& ObjectPointsMat,
 		cv::Mat& ImagePointsMat);
 
@@ -220,4 +229,22 @@ private:
 	int LambdaLog10 = -3;
 };
 
+class FOptimizeNodalOffsetSolver : public cv::MinProblemSolver::Function
+{
+public:
+
+	//~ Begin cv::MinProblemSolver::Function interface
+	int getDims() const;
+	double calc(const double* x) const;
+	//~ End cv::MinProblemSolver::Function interface
+
+public:
+	FVector2D FocalLength;
+	FVector2D ImageCenter;
+
+	TArray<FTransform> CameraPoses;
+	TArray<TArray<FVector>> Points3d;
+	TArray<TArray<FVector2f>> Points2d;
+};
+
 #endif // WITH_OPENCV
diff --git a/Engine/Plugins/VirtualProduction/CameraCalibration/Source/CameraCalibrationEditor/Private/Calibrators/CameraLensDistortionAlgoCheckerboard.cpp b/Engine/Plugins/VirtualProduction/CameraCalibration/Source/CameraCalibrationEditor/Private/Calibrators/CameraLensDistortionAlgoCheckerboard.cpp
--- a/Engine/Plugins/VirtualProduction/CameraCalibration/Source/CameraCalibrationEditor/Private/Calibrators/CameraLensDistortionAlgoCheckerboard.cpp
+++ b/Engine/Plugins/VirtualProduction/CameraCalibration/Source/CameraCalibrationEditor/Private/Calibrators/CameraLensDistortionAlgoCheckerboard.cpp
@@ -5,12 +5,12 @@
 #include "AssetEditor/CameraCalibrationStepsController.h"
 #include "AssetEditor/LensDistortionTool.h"
 #include "AssetEditor/SSimulcamViewport.h"
-#include "AssetRegistry/AssetData.h"
 #include "CalibrationPointComponent.h"
 #include "Camera/CameraActor.h"
 #include "CameraCalibrationCheckerboard.h"
 #include "CameraCalibrationEditorLog.h"
 #include "CameraCalibrationSolver.h"
+#include "CameraCalibrationUtilsPrivate.h"
 #include "Dom/JsonObject.h"
 #include "EditorFontGlyphs.h"
 #include "Engine/Texture2D.h"
@@ -18,10 +18,14 @@
 #include "ImageUtils.h"
 #include "JsonObjectConverter.h"
 #include "LensComponent.h"
+#include "LensFile.h"
 #include "Misc/MessageDialog.h"
+#include "Materials/MaterialInstanceDynamic.h"
 #include "Models/AnamorphicLensModel.h"
+#include "Models/SphericalLensModel.h"
+#include "OpenCVHelper.h"
 #include "PropertyCustomizationHelpers.h"
-#include "SphericalLensDistortionModelHandler.h"
+#include "TextureResource.h"
 #include "UI/CameraCalibrationWidgetHelpers.h"
 #include "Widgets/Input/SButton.h"
 #include "Widgets/Input/SCheckBox.h"
@@ -47,6 +51,21 @@ namespace UE::CameraCalibration::Private::LensDistortionCheckerboardExportFields
 	static const FString Version(TEXT("Version"));
 }
 
+namespace UE::CameraCalibration::Private
+{
+	void AdjustCorners(TArray<FVector2f>& Corners, FIntPoint TextureSize, FIntPoint ImageSize)
+	{
+		// It is possible that the size of the debug texture where we want to draw the checkerboard corners is different than the size of the image where the checkerboard corners were found.
+		// Before drawing, all of the corner positions should be shifted to account for the difference in size.
+		const FVector2f TopLeftCorner = (TextureSize - ImageSize) / 2.0f;
+
+		for (FVector2f& Corner : Corners)
+		{
+			Corner += TopLeftCorner;
+		}
+	}
+}
+
 namespace CameraLensDistortionAlgoCheckerboard
 {
 	class SCalibrationRowGenerator
@@ -124,7 +143,6 @@ void UCameraLensDistortionAlgoCheckerboard::Initialize(ULensDistortionTool* InTo
 	// Guess which calibrator to use.
 	SetCalibrator(FindFirstCalibrator());
 
-#if WITH_OPENCV
 	// Initialize coverage matrix
 	FCameraCalibrationStepsController* StepsController = Tool->GetCameraCalibrationStepsController();
 
@@ -133,10 +151,10 @@ void UCameraLensDistortionAlgoCheckerboard::Initialize(ULensDistortionTool* InTo
 		return;
 	}
 
-	FIntPoint Size = StepsController->GetCompRenderResolution();
-	CvCoverage = cv::Mat(cv::Size(Size.X, Size.Y), CV_8UC4);
-	CoverageTexture = FOpenCVHelper::TextureFromCvMat(CvCoverage, CoverageTexture);
-#endif
+	const FIntPoint Size = StepsController->GetCompRenderResolution();
+
+	CoverageTexture = UTexture2D::CreateTransient(Size.X, Size.Y, EPixelFormat::PF_B8G8R8A8);
+	UE::CameraCalibration::Private::ClearTexture(CoverageTexture);
 
 	if (UMaterialInstanceDynamic* OverlayMID = Tool->GetOverlayMID())
 	{
@@ -169,6 +187,14 @@ void UCameraLensDistortionAlgoCheckerboard::Tick(float DeltaTime)
 		return;
 	}
 
+	// If the resolution of the simulcam comp has changed, update the coverage texture to be the correct size
+	const FIntPoint Size = StepsController->GetCompRenderResolution();
+	if (CoverageTexture && ((CoverageTexture->GetSizeX() != Size.X) || (CoverageTexture->GetSizeY() != Size.Y)))
+	{
+		CoverageTexture = UTexture2D::CreateTransient(Size.X, Size.Y, EPixelFormat::PF_B8G8R8A8);
+		RefreshCoverage();
+	}
+
 	// If not paused, cache calibrator 3d point position
 	if (!StepsController->IsPaused())
 	{
@@ -224,19 +250,17 @@ bool UCameraLensDistortionAlgoCheckerboard::OnViewportClicked(const FGeometry& M
 	}
 
 	// Left Mouse button means to add a new calibration row
-	{
-		FText ErrorMessage;
+	FText ErrorMessage;
 
-		if (!AddCalibrationRow(ErrorMessage))
-		{
-			FMessageDialog::Open(EAppMsgType::Ok, ErrorMessage);
-		}
+	if (!AddCalibrationRow(ErrorMessage))
+	{
+		FMessageDialog::Open(EAppMsgType::Ok, ErrorMessage);
+	}
 
-		// force play (whether we succeeded or not)
-		if (FCameraCalibrationStepsController* StepsController = GetStepsController())
-		{
-			StepsController->Play();
-		}
+	// force play (whether we succeeded or not)
+	if (FCameraCalibrationStepsController* StepsController = GetStepsController())
+	{
+		StepsController->Play();
 	}
 
 	return true;
@@ -254,7 +278,6 @@ FCameraCalibrationStepsController* UCameraLensDistortionAlgoCheckerboard::GetSte
 
 bool UCameraLensDistortionAlgoCheckerboard::AddCalibrationRow(FText& OutErrorMessage)
 {
-#if WITH_OPENCV
 	using namespace CameraLensDistortionAlgoCheckerboard;
 
 	if (!Tool.IsValid())
@@ -300,12 +323,6 @@ bool UCameraLensDistortionAlgoCheckerboard::AddCalibrationRow(FText& OutErrorMes
 		return false;
 	}
 
-	// Create OpenCV Mat with those pixels
-	cv::Mat CvFrame(cv::Size(Size.X, Size.Y), CV_8UC4, Pixels.GetData());
-
-	cv::Mat CvGray;
-	cv::cvtColor(CvFrame, CvGray, cv::COLOR_RGBA2GRAY);
-
 	// Export the latest session data
 	ExportSessionData();
 
@@ -317,99 +334,31 @@ bool UCameraLensDistortionAlgoCheckerboard::AddCalibrationRow(FText& OutErrorMes
 
 	Row->Index = RowIndex;
 	Row->CameraData = LastCameraData;
-	Row->ImageHeight = CvGray.rows;
-	Row->ImageWidth = CvGray.cols;
+	Row->ImageWidth = Size.X;
+	Row->ImageHeight = Size.Y;
 
 	Row->NumCornerRows = Calibrator->NumCornerRows;
 	Row->NumCornerCols = Calibrator->NumCornerCols;
 	Row->SquareSideInCm = Calibrator->SquareSideLength;
 
-	// Identify checkerboard
-	{
-		cv::Size CheckerboardSize(Row->NumCornerCols, Row->NumCornerRows);
+	const FIntPoint CheckerboardDimensions = FIntPoint(Calibrator->NumCornerCols, Calibrator->NumCornerRows);
 
-		std::vector<cv::Point2f> Corners;
+	const bool bCornersFound = FOpenCVHelper::IdentifyCheckerboard(Pixels, Size, CheckerboardDimensions, Row->Points2d);
 
-		const bool bCornersFound = cv::findChessboardCorners(
-			CvGray,
-			CheckerboardSize,
-			Corners,
-			cv::CALIB_CB_ADAPTIVE_THRESH | cv::CALIB_CB_NORMALIZE_IMAGE
-		);
-
-		if (!bCornersFound || Corners.empty())
-		{
-			OutErrorMessage = FText::FromString(FString::Printf(TEXT(
-				"Could not identify the expected checkerboard points of interest. "
-				"The expected checkerboard has %dx%d inner corners."),
-				Row->NumCornerCols, Row->NumCornerRows)
-			);
-
-			return false;
-		}
-
-		// cv::TermCriteria::Type::EPS will stop the search when the error is under the given epsilon.
-		// cv::TermCriteria::Type::COUNT will stop after the specified number of iterations regardless of epsilon.
-		cv::TermCriteria Criteria(cv::TermCriteria::Type::EPS | cv::TermCriteria::Type::COUNT, 30, 0.001);
-		cv::cornerSubPix(CvGray, Corners, cv::Size(11, 11), cv::Size(-1, -1), Criteria);
-
-		// Force first one to be top-left
-		if (Corners.front().y > Corners.back().y)
-		{
-			std::reverse(Corners.begin(), Corners.end());
-		}
-
-		for (const cv::Point2f& Corner : Corners)
-		{
-			Row->Points2d.Add(FVector2D(Corner.x, Corner.y));
-		}
-
-		// Save an image view of the captured frame with the corners overlaid on it (for exporting)
-		FImageView ImageView = FImageView(CvFrame.data, CvFrame.cols, CvFrame.rows, ERawImageFormat::BGRA8);
-		Row->ImageView = ImageView;
-
-		// Show the detection to the user
-		if (bShouldShowDetectionWindow)
-		{
-			cv::Mat CvFrameWithOverlay;
-			CvFrame.copyTo(CvFrameWithOverlay);
-			cv::drawChessboardCorners(CvFrameWithOverlay, CheckerboardSize, Corners, true);
-
-			FCameraCalibrationWidgetHelpers::DisplayTextureInWindowAlmostFullScreen(
-				FOpenCVHelper::TextureFromCvMat(CvFrameWithOverlay),
-				LOCTEXT("CheckerboardDetection", "Checkerboard Detection")
-			);
-		}
-
-		// The corners were found using only the camera feed, but the coverage overlay will be drawn over the entire comp
-		// Before updating the overlay texture, shift all the camera feed corner pixels to represent the correct position in the overlay texture
-		const FIntPoint CompSize = StepsController->GetCompRenderResolution();
-		const FIntPoint CameraFeedSize = StepsController->GetCameraFeedSize();
-
-		const FIntPoint TopLeftCorner = (CompSize - CameraFeedSize) / 2;
-
-		for (cv::Point2f& Corner : Corners)
-		{
-			Corner.x += TopLeftCorner.X;
-			Corner.y += TopLeftCorner.Y;
-		}
-
-		// Refresh the coverage overlay if the comp resolution has changed since the last image was taken
-		if ((CvCoverage.cols != CompSize.X) || (CvCoverage.rows != CompSize.Y))
-		{
-			RefreshCoverage();
-		}
-
-		// Update the coverage overlay image with information from the newly added row
-		cv::drawChessboardCorners(CvCoverage, CheckerboardSize, Corners, true);
-
-		CoverageTexture = FOpenCVHelper::TextureFromCvMat(CvCoverage, CoverageTexture);
-		StepsController->RefreshOverlay();
-	}
+ 	if (!bCornersFound || Row->Points2d.IsEmpty())
+ 	{
+ 		OutErrorMessage = FText::FromString(FString::Printf(TEXT(
+ 			"Could not identify the expected checkerboard points of interest. "
+ 			"The expected checkerboard has %dx%d inner corners."),
+ 			Row->NumCornerCols, Row->NumCornerRows)
+ 		);
+ 
+ 		return false;
+ 	}
 
+	// Fill out the checkerboard's 3D points
 	if (CVarUseExtrinsicsGuess.GetValueOnGameThread())
 	{
-		// Fill out checkerboard 3d points
 		const FVector LocationTL = Calibrator->TopLeft->GetComponentLocation();
 		const FVector LocationTR = Calibrator->TopRight->GetComponentLocation();
 		const FVector LocationBL = Calibrator->BottomLeft->GetComponentLocation();
@@ -417,7 +366,6 @@ bool UCameraLensDistortionAlgoCheckerboard::AddCalibrationRow(FText& OutErrorMes
 		const FVector RightVector = LocationTR - LocationTL;
 		const FVector DownVector = LocationBL - LocationTL;
 
-		// Fill out checkerboard 3d points
 		const float HorizontalStep = (Row->NumCornerCols > 1) ? (1.0f / (Row->NumCornerCols - 1)) : 0.0f;
 		const float VerticalStep = (Row->NumCornerRows > 1) ? (1.0f / (Row->NumCornerRows - 1)) : 0.0f;
 
@@ -426,19 +374,12 @@ bool UCameraLensDistortionAlgoCheckerboard::AddCalibrationRow(FText& OutErrorMes
 			for (int32 ColIdx = 0; ColIdx < Row->NumCornerCols; ++ColIdx)
 			{
 				const FVector PointLocation = LocationTL + (RightVector * ColIdx * HorizontalStep) + (DownVector * RowIdx * VerticalStep);
-
-				// Convert from UE coordinates to OpenCV coordinates
-				FTransform Transform = FTransform::Identity;
-				Transform.SetLocation(PointLocation);
-				FOpenCVHelper::ConvertUnrealToOpenCV(Transform);
-
-				Row->Points3d.Add(Transform.GetLocation());
+				Row->Points3d.Add(PointLocation);
 			}
 		}
 	}
 	else
 	{
-		// Fill out checkerboard 3d points
 		for (int32 RowIdx = 0; RowIdx < Row->NumCornerRows; ++RowIdx)
 		{
 			for (int32 ColIdx = 0; ColIdx < Row->NumCornerCols; ++ColIdx)
@@ -448,26 +389,43 @@ bool UCameraLensDistortionAlgoCheckerboard::AddCalibrationRow(FText& OutErrorMes
 		}
 	}
 
-	// Create thumbnail and add it to the row
+	// Save an image view of the captured frame (for exporting)
+	FImageView ImageView = FImageView(Pixels.GetData(), Size.X, Size.Y, ERawImageFormat::BGRA8);
+	Row->ImageView = ImageView;
+
+	// Generate a thumbnail image for display in the tool for this row
+	FImage RowThumbnail;
+	constexpr int32 ResolutionDivider = 4;
+	FImageCore::ResizeTo(ImageView, RowThumbnail, Size.X / ResolutionDivider, Size.Y / ResolutionDivider, ERawImageFormat::BGRA8, ImageView.GetGammaSpace());
+
+	UTexture2D* ThumbnailTexture = FImageUtils::CreateTexture2DFromImage(RowThumbnail);
+	if (ThumbnailTexture)
 	{
-		// Resize the frame to thumbnail size
+		Row->Thumbnail = SNew(SSimulcamViewport, ThumbnailTexture).WithZoom(false).WithPan(false);
+	}
 
-		cv::Mat CvThumbnail;
-		const int32 ResolutionDivider = 4;
-		cv::resize(CvFrame, CvThumbnail, cv::Size(CvFrame.cols / ResolutionDivider, CvFrame.rows / ResolutionDivider));
+	// Update the coverage overlay with the latest checkerboard corners
+	TArray<FVector2f> CameraFeedAdjustedCorners = Row->Points2d;
+	const FIntPoint CoverageTextureSize = FIntPoint(CoverageTexture->GetSizeX(), CoverageTexture->GetSizeY());
+	UE::CameraCalibration::Private::AdjustCorners(CameraFeedAdjustedCorners, CoverageTextureSize, Size);
 
-		if (UTexture2D* ThumbnailTexture = FOpenCVHelper::TextureFromCvMat(CvThumbnail))
-		{
-			Row->Thumbnail = SNew(SSimulcamViewport, ThumbnailTexture).WithZoom(false).WithPan(false);
-		}
+	FOpenCVHelper::DrawCheckerboardCorners(CameraFeedAdjustedCorners, CheckerboardDimensions, CoverageTexture);
+	StepsController->RefreshOverlay();
+
+	// Show the detection to the user
+	if (bShouldShowDetectionWindow)
+	{
+		UTexture2D* DebugTexture = UTexture2D::CreateTransient(Size.X, Size.Y, EPixelFormat::PF_B8G8R8A8);
+		UE::CameraCalibration::Private::SetTextureData(DebugTexture, Pixels);
+		FOpenCVHelper::DrawCheckerboardCorners(Row->Points2d, CheckerboardDimensions, DebugTexture);
+
+		FCameraCalibrationWidgetHelpers::DisplayTextureInWindowAlmostFullScreen(DebugTexture, LOCTEXT("CheckerboardDetection", "Checkerboard Detection"));
 	}
 
 	// Validate the new row, show a message if validation fails.
+	if (!ValidateNewRow(Row, OutErrorMessage))
 	{
-		if (!ValidateNewRow(Row, OutErrorMessage))
-		{
-			return false;
-		}
+		return false;
 	}
 
 	// Add this data point
@@ -484,10 +442,6 @@ bool UCameraLensDistortionAlgoCheckerboard::AddCalibrationRow(FText& OutErrorMes
 	}
 
 	return true;
-#endif //WITH_OPENCV
-
-	OutErrorMessage = LOCTEXT("OpenCVRequired", "OpenCV is required");
-	return false;
 }
 
 TSharedRef<SWidget> UCameraLensDistortionAlgoCheckerboard::BuildUI()
@@ -720,28 +674,32 @@ bool UCameraLensDistortionAlgoCheckerboard::GetLensDistortion(
 		return false;
 	}
 
-#if WITH_OPENCV
-
-	// Initialize the camera matrix that will be used in each call to projectPoints()
-	float TestImageWidth = LensFile->CameraFeedInfo.GetDimensions().X;
-	float TestImageHeight = LensFile->CameraFeedInfo.GetDimensions().Y;
-
-	float TestSensorWidth = StepsController->GetLensFileEvaluationInputs().Filmback.SensorWidth;
+	const FIntPoint ImageSize = LensFile->CameraFeedInfo.GetDimensions();
 
 	const float PixelAspect = LensFile->LensInfo.SqueezeFactor;
 
-	TestSensorWidth = TestSensorWidth * PixelAspect;
+	if (FMath::IsNearlyZero(PixelAspect))
+	{
+		OutErrorMessage = LOCTEXT("PixelAspectNearlyZero", "The pixel aspect in the Lens Information is zero, which is invalid.");
+		return false;
+	}
+
+	const float PhysicalSensorWidth = StepsController->GetLensFileEvaluationInputs().Filmback.SensorWidth;
+	const float DesqueezeSensorWidth = PhysicalSensorWidth * PixelAspect;
 
-	float Fx = TestImageWidth;
-	if (!FMath::IsNearlyZero(TestSensorWidth))
+	if (FMath::IsNearlyZero(DesqueezeSensorWidth))
 	{
-		Fx = (FocalLengthEstimate / TestSensorWidth) * TestImageWidth;
+		OutErrorMessage = LOCTEXT("SensorWidthNearlyZero", "The sensor width of the CineCamera is zero, which is invalid.");
+		return false;
 	}
 
-	FVector2f FocalLength = FVector2f(Fx, Fx);
-	FVector2f ImageCenter = FVector2f((TestImageWidth - 1) * 0.5f, (TestImageHeight - 1) * 0.5f);
+	const double Fx = (FocalLengthEstimate / DesqueezeSensorWidth) * ImageSize.X;
+
+	// When operating on a desqueezed image, we expect our pixel aspect to be square, so horizontal and vertical field of view are assumed to be equal (i.e. Fx == Fy)
+	FVector2D FocalLength = FVector2D(Fx, Fx); 
+	FVector2D ImageCenter = FVector2D((ImageSize.X - 1) * 0.5, (ImageSize.Y - 1) * 0.5);
 
-	TArray<TArray<FVector2D>> Samples2d;
+	TArray<TArray<FVector2f>> Samples2d;
 	TArray<TArray<FVector>> Samples3d;
 
 	TArray<FTransform> CameraPoses;
@@ -751,7 +709,6 @@ bool UCameraLensDistortionAlgoCheckerboard::GetLensDistortion(
 		Samples2d.Add(Row->Points2d);
 		Samples3d.Add(Row->Points3d);
 
-		FOpenCVHelper::ConvertUnrealToOpenCV(Row->CameraData.Pose);
 		CameraPoses.Add(Row->CameraData.Pose);
 	}
 
@@ -793,7 +750,7 @@ bool UCameraLensDistortionAlgoCheckerboard::GetLensDistortion(
 		LensFile->LensInfo.LensModel,
 		Samples3d,
 		Samples2d,
-		FIntPoint(LastRow->ImageWidth, LastRow->ImageHeight),
+		ImageSize,
 		FocalLength,
 		ImageCenter,
 		DistortionCoefficients,
@@ -803,19 +760,12 @@ bool UCameraLensDistortionAlgoCheckerboard::GetLensDistortion(
 	);
 
 	OutLensModel = LensFile->LensInfo.LensModel;
+
 	OutDistortionInfo.Parameters = DistortionCoefficients;
 
-	// FocalLengthInfo
-	OutFocalLengthInfo.FxFy = FVector2D(
-		float(FocalLength.X / LastRow->ImageWidth),
-		float(FocalLength.Y / LastRow->ImageHeight)
-	);
+	OutFocalLengthInfo.FxFy = FocalLength / ImageSize;
 
-	// ImageCenterInfo
-	OutImageCenterInfo.PrincipalPoint = FVector2D(
-		float(ImageCenter.X / LastRow->ImageWidth),
-		float(ImageCenter.Y / LastRow->ImageHeight)
-	);
+	OutImageCenterInfo.PrincipalPoint = ImageCenter / ImageSize;
 
 	// FZ inputs to LUT
 	OutFocus = LastRow->CameraData.InputFocus;
@@ -842,8 +792,6 @@ bool UCameraLensDistortionAlgoCheckerboard::GetLensDistortion(
 
 		for (int32 RowIndex = 0; RowIndex < CalibrationRows.Num(); ++RowIndex)
 		{
-			FOpenCVHelper::ConvertOpenCVToUnreal(CalibrationRows[RowIndex]->CameraData.Pose);
-
 			const FTransform DesiredCameraTransform = CameraPoses[RowIndex];
 			const FTransform DesiredOffset = DesiredCameraTransform * CalibrationRows[RowIndex]->CameraData.Pose.Inverse() * ExistingOffset;
 
@@ -858,12 +806,6 @@ bool UCameraLensDistortionAlgoCheckerboard::GetLensDistortion(
 	}
 
 	return true;
-#else
-	{
-		OutErrorMessage = LOCTEXT("OpenCVRequired", "OpenCV is required");
-		return false;
-	}
-#endif //WITH_OPENCV
 }
 
 TSharedRef<SWidget> UCameraLensDistortionAlgoCheckerboard::BuildCalibrationDevicePickerWidget()
@@ -969,9 +911,9 @@ TSharedRef<SWidget> UCameraLensDistortionAlgoCheckerboard::BuildFocalLengthEstim
 	.FillWidth(0.8)
 	.Padding(0, 0, 10, 0)
 	[
-		SNew(SNumericEntryBox<float>)
-		.Value(MakeAttributeLambda([this]() { return TOptional<float>(FocalLengthEstimate); }))
-		.OnValueChanged(SNumericEntryBox<float>::FOnValueChanged::CreateLambda([this](float NewValue) { FocalLengthEstimate = NewValue; }))
+		SNew(SNumericEntryBox<double>)
+		.Value(MakeAttributeLambda([this]() { return TOptional<double>(FocalLengthEstimate); }))
+		.OnValueChanged(SNumericEntryBox<double>::FOnValueChanged::CreateLambda([this](double NewValue) { FocalLengthEstimate = NewValue; }))
 	]
 
 	+ SHorizontalBox::Slot()
@@ -1194,7 +1136,6 @@ void UCameraLensDistortionAlgoCheckerboard::ClearCalibrationRows()
 
 void UCameraLensDistortionAlgoCheckerboard::RefreshCoverage()
 {
-#if WITH_OPENCV
 	FCameraCalibrationStepsController* StepsController = Tool->GetCameraCalibrationStepsController();
 
 	if (!ensure(StepsController))
@@ -1202,27 +1143,19 @@ void UCameraLensDistortionAlgoCheckerboard::RefreshCoverage()
 		return;
 	}
 
-	FIntPoint Size = StepsController->GetCompRenderResolution();
-
-	CvCoverage.release();
-	CvCoverage = cv::Mat(cv::Size(Size.X, Size.Y), CV_8UC4);
+	UE::CameraCalibration::Private::ClearTexture(CoverageTexture);
 
 	for (const TSharedPtr<FLensDistortionCheckerboardRowData>& Row : CalibrationRows)
 	{
-		cv::Size CheckerboardSize(Row->NumCornerCols, Row->NumCornerRows);
-
-		std::vector<cv::Point2f> Corners;
+		TArray<FVector2f> CameraFeedAdjustedCorners = Row->Points2d;
+		const FIntPoint CoverageTextureSize = FIntPoint(CoverageTexture->GetSizeX(), CoverageTexture->GetSizeY());
+		const FIntPoint ImageSize = FIntPoint(Row->ImageWidth, Row->ImageHeight);
+		UE::CameraCalibration::Private::AdjustCorners(CameraFeedAdjustedCorners, CoverageTextureSize, ImageSize);
 
-		for (const FVector2D& Corner : Row->Points2d)
-		{
-			Corners.push_back(cv::Point2f(Corner.X, Corner.Y));
-		}
-
-		cv::drawChessboardCorners(CvCoverage, CheckerboardSize, Corners, true);
+		const FIntPoint CheckerboardDimensions = FIntPoint(Row->NumCornerCols, Row->NumCornerRows);
+		FOpenCVHelper::DrawCheckerboardCorners(Row->Points2d, CheckerboardDimensions, CoverageTexture);
 	}
 
-	CoverageTexture = FOpenCVHelper::TextureFromCvMat(CvCoverage, CoverageTexture);
-
 	// The coverage texture may have changed as a result of a change in size or pixel format.
 	// Therefore, the material parameter should be updated to ensure it is up to date.
 	if (CoverageTexture)
@@ -1234,7 +1167,6 @@ void UCameraLensDistortionAlgoCheckerboard::RefreshCoverage()
 	}
 
 	StepsController->RefreshOverlay();
-#endif
 }
 
 
diff --git a/Engine/Plugins/VirtualProduction/CameraCalibration/Source/CameraCalibrationEditor/Private/Calibrators/CameraLensDistortionAlgoCheckerboard.h b/Engine/Plugins/VirtualProduction/CameraCalibration/Source/CameraCalibrationEditor/Private/Calibrators/CameraLensDistortionAlgoCheckerboard.h
--- a/Engine/Plugins/VirtualProduction/CameraCalibration/Source/CameraCalibrationEditor/Private/Calibrators/CameraLensDistortionAlgoCheckerboard.h
+++ b/Engine/Plugins/VirtualProduction/CameraCalibration/Source/CameraCalibrationEditor/Private/Calibrators/CameraLensDistortionAlgoCheckerboard.h
@@ -3,21 +3,7 @@
 #pragma once
 
 #include "CameraLensDistortionAlgo.h"
-
-#include "LensFile.h"
-
-#if WITH_OPENCV
-
-#include <vector>
-
-#include "OpenCVHelper.h"
-#include "PreOpenCVHeaders.h"
-#include "opencv2/calib3d.hpp"
-#include "opencv2/imgproc.hpp"
-#include "PostOpenCVHeaders.h"
-
-#endif //WITH_OPENCV
-
+#include "ImageCore.h"
 
 #include "CameraLensDistortionAlgoCheckerboard.generated.h"
 
@@ -90,7 +76,7 @@ struct FLensDistortionCheckerboardRowData
 
 	// Checkerboard corners in 2d image pixel coordinates.
 	UPROPERTY()
-	TArray<FVector2D> Points2d;
+	TArray<FVector2f> Points2d;
 
 	// Checkerboard corners in 3d local space.
 	UPROPERTY()
@@ -201,7 +187,7 @@ private:
 	bool bFixImageCenter = false;
 
 	/** Estimated focal length to initialize the distortion solver */
-	float FocalLengthEstimate = 0.0f;
+	double FocalLengthEstimate = 0.0;
 
 	/** If true, the solver will use the current camera pose to initialize the camera extrinsic parameters for each image */
 	bool bUseExtrinsicsGuess = false;
@@ -212,11 +198,6 @@ private:
 	/** Texture into which detected chessboard corners from each calibration row are drawn */
 	TObjectPtr<UTexture2D> CoverageTexture;
 
-#if WITH_OPENCV
-	/** OpenCV matrix used to draw chessboard corners */
-	cv::Mat CvCoverage;
-#endif
-
 private:
 
 	/** Builds the UI of the calibration device picker */
diff --git a/Engine/Plugins/VirtualProduction/CameraCalibration/Source/CameraCalibrationEditor/Private/Calibrators/CameraLensDistortionAlgoPoints.cpp b/Engine/Plugins/VirtualProduction/CameraCalibration/Source/CameraCalibrationEditor/Private/Calibrators/CameraLensDistortionAlgoPoints.cpp
--- a/Engine/Plugins/VirtualProduction/CameraCalibration/Source/CameraCalibrationEditor/Private/Calibrators/CameraLensDistortionAlgoPoints.cpp
+++ b/Engine/Plugins/VirtualProduction/CameraCalibration/Source/CameraCalibrationEditor/Private/Calibrators/CameraLensDistortionAlgoPoints.cpp
@@ -7,6 +7,7 @@
 #include "CalibrationPointComponent.h"
 #include "Camera/CameraActor.h"
 #include "CameraCalibrationEditorLog.h"
+#include "CameraCalibrationSolver.h"
 #include "Dom/JsonObject.h"
 #include "Editor.h"
 #include "Input/Events.h"
@@ -14,6 +15,7 @@
 #include "Math/Vector.h"
 #include "Misc/MessageDialog.h"
 #include "Models/SphericalLensModel.h"
+#include "OpenCVHelper.h"
 #include "UI/SFilterableActorPicker.h"
 #include "UI/CameraCalibrationWidgetHelpers.h"
 #include "UObject/UObjectIterator.h"
@@ -21,16 +23,6 @@
 #include "Widgets/Input/SNumericEntryBox.h"
 #include "Widgets/Views/SListView.h"
 
-#if WITH_OPENCV
-
-#include "OpenCVHelper.h"
-#include <vector>
-
-#include "PreOpenCVHeaders.h"
-#include "opencv2/calib3d.hpp"
-#include "PostOpenCVHeaders.h"
-
-#endif
 
 #define LOCTEXT_NAMESPACE "CameraLensDistortionAlgoPoints"
 
@@ -264,6 +256,10 @@ bool UCameraLensDistortionAlgoPoints::OnViewportClicked(const FGeometry& MyGeome
 		return true;
 	}
 
+	// Store the 2D point in pixel coordinates
+	const FIntPoint ImageSize = StepsController->GetCompRenderResolution();
+	LastCalibratorPoint.Point2d = LastCalibratorPoint.Point2d * ImageSize;
+
 	// Export the latest session data
 	ExportSessionData();
 
@@ -364,7 +360,6 @@ bool UCameraLensDistortionAlgoPoints::GetLensDistortion(
 	double& OutError,
 	FText& OutErrorMessage)
 {
-#if WITH_OPENCV
 	if (!LensDistortionTool.IsValid())
 	{
 		OutErrorMessage = LOCTEXT("InvalidTool", "Invalid Tool");
@@ -387,38 +382,32 @@ bool UCameraLensDistortionAlgoPoints::GetLensDistortion(
 		return false;
 	}
 
-	const FIntPoint ImageSize = StepsController->GetCompRenderResolution();
+	TArray<TArray<FVector>> Samples3d;
+	TArray<TArray<FVector2f>> Samples2d;
 
-	// Gather 3D points and 2D points from each group of rows
-	std::vector<std::vector<cv::Point3f>> Samples3d;
-	std::vector<std::vector<cv::Point2f>> Samples2d;
+	TArray<FTransform> CameraPoses;
 
 	uint32 NextPatternIndex = 0;
 	for (int32 RowIndex = 0; RowIndex < CalibrationRows.Num(); ++RowIndex)
 	{
-		std::vector<cv::Point3f> Points3d;
-		std::vector<cv::Point2f> Points2d;
+		TArray<FVector> Points3d;
+		TArray<FVector2f> Points2d;
 
 		while (RowIndex < CalibrationRows.Num() && CalibrationRows[RowIndex]->PatternIndex == NextPatternIndex)
 		{
 			const TSharedPtr<FLensDistortionPointsRowData>& Row = CalibrationRows[RowIndex];
 
-			Points3d.push_back(cv::Point3f(
-				Row->CalibratorPointData.Point3d.X,
-				Row->CalibratorPointData.Point3d.Y,
-				Row->CalibratorPointData.Point3d.Z));
-
-			Points2d.push_back(cv::Point2f(
-				Row->CalibratorPointData.Point2d.X,
-				Row->CalibratorPointData.Point2d.Y));
+			Points3d.Add(Row->CalibratorPointData.Point3d);
+			Points2d.Add(Row->CalibratorPointData.Point2d);
 
 			++RowIndex;
 		}
 
-		if (Points3d.size() > 0)
+		if (Points3d.Num() > 0)
 		{
-			Samples3d.push_back(Points3d);
-			Samples2d.push_back(Points2d);
+			Samples3d.Add(Points3d);
+			Samples2d.Add(Points2d);
+			CameraPoses.Add(FTransform::Identity);
 		}
 
 		--RowIndex;
@@ -426,17 +415,17 @@ bool UCameraLensDistortionAlgoPoints::GetLensDistortion(
 	}
 
 	// Validate that there are at least 4 patterns
-	if (Samples3d.size() < 4)
+	if (Samples3d.Num() < 4)
 	{
 		OutErrorMessage = LOCTEXT("NotEnoughSamples", "At least 4 calibration patterns are required");
 		return false;
 	}
 
 	// Validate that each pattern has the same number of points
-	const int32 NumPointsInPattern = Samples3d[0].size();
-	for (int32 PatternIndex = 1; PatternIndex < Samples3d.size(); ++PatternIndex)
+	const int32 NumPointsInPattern = Samples3d[0].Num();
+	for (int32 PatternIndex = 1; PatternIndex < Samples3d.Num(); ++PatternIndex)
 	{
-		if (Samples3d[PatternIndex].size() != NumPointsInPattern)
+		if (Samples3d[PatternIndex].Num() != NumPointsInPattern)
 		{
 			OutErrorMessage = LOCTEXT("DifferentNumPointsInPattern", "Every calibration pattern must have the same number of points");
 			return false;
@@ -444,7 +433,7 @@ bool UCameraLensDistortionAlgoPoints::GetLensDistortion(
 	}
 
 	// Because the calibration pattern is not coplanar, OpenCV requires an initial intrinsics guess
-	if (FMath::IsNearlyEqual(CurrentFocalLengthMM, 0.0f) || CurrentFocalLengthMM < 0.0f)
+	if (FMath::IsNearlyEqual(FocalLengthEstimate, 0.0) || FocalLengthEstimate < 0.0)
 	{
 		OutErrorMessage = LOCTEXT("InvalidFocalLengthMsg", "The current focal length (in mm) must be set to a "
 			"valid value in order to seed the calibration algorithm with a best guess as to the camera intrinsics (Fx/Fy)");
@@ -452,92 +441,63 @@ bool UCameraLensDistortionAlgoPoints::GetLensDistortion(
 	}
 
 	// Validate sensor dimensions
-	constexpr float MinimumSize = 1.0f; // Limit sensor dimension to 1mm
-	if (LensFile->LensInfo.SensorDimensions.X < MinimumSize || LensFile->LensInfo.SensorDimensions.Y < MinimumSize)
+	const FIntPoint ImageSize = LensFile->CameraFeedInfo.GetDimensions();
+
+	const float PixelAspect = LensFile->LensInfo.SqueezeFactor;
+
+	if (FMath::IsNearlyZero(PixelAspect))
 	{
-		OutErrorMessage = LOCTEXT("InvalidSensorDimensions", "Invalid sensor dimensions. Can't have dimensions smaller than 1mm");
+		OutErrorMessage = LOCTEXT("PixelAspectNearlyZero", "The pixel aspect in the Lens Information is zero, which is invalid.");
 		return false;
 	}
 
-	const float Fx = CurrentFocalLengthMM / LensFile->LensInfo.SensorDimensions.X;
-	const float Fy = CurrentFocalLengthMM / LensFile->LensInfo.SensorDimensions.Y;
+	const float PhysicalSensorWidth = StepsController->GetLensFileEvaluationInputs().Filmback.SensorWidth;
+	const float DesqueezeSensorWidth = PhysicalSensorWidth * PixelAspect;
 
-	cv::Mat CameraMatrix = cv::Mat::eye(3, 3, CV_64F);
+	if (FMath::IsNearlyZero(DesqueezeSensorWidth))
+	{
+		OutErrorMessage = LOCTEXT("SensorWidthNearlyZero", "The sensor width of the CineCamera is zero, which is invalid.");
+		return false;
+	}
 
-	CameraMatrix.at<double>(0, 0) = Fx * ImageSize.X; // Fx
-	CameraMatrix.at<double>(1, 1) = Fy * ImageSize.Y; // Fy
-	CameraMatrix.at<double>(0, 2) = 0.5 * ImageSize.X; // Cx
-	CameraMatrix.at<double>(1, 2) = 0.5 * ImageSize.Y; // Cy
+	const double Fx = (FocalLengthEstimate / DesqueezeSensorWidth) * ImageSize.X;
 
-	cv::Mat DistortionCoefficients;
+	// When operating on a desqueezed image, we expect our pixel aspect to be square, so horizontal and vertical field of view are assumed to be equal (i.e. Fx == Fy)
+	FVector2D FocalLength = FVector2D(Fx, Fx);
+	FVector2D ImageCenter = FVector2D((ImageSize.X - 1) * 0.5, (ImageSize.Y - 1) * 0.5);
 
-	std::vector<cv::Mat> Rvecs;
-	std::vector<cv::Mat> Tvecs;
+	ECalibrationFlags SolverFlags = ECalibrationFlags::None;
+	EnumAddFlags(SolverFlags, ECalibrationFlags::UseIntrinsicGuess);
 
- 	std::vector<double> IntrinsicsDeviations;
- 	std::vector<double> ExtrinsicsDeviations;
- 	std::vector<double> ViewErrors;
+	TArray<float> DistortionCoefficients;
 
-	OutError = cv::calibrateCamera(
+	OutError = FCameraCalibrationSolver::CalibrateCamera(
+		LensFile->LensInfo.LensModel,
 		Samples3d,
 		Samples2d,
-		cv::Size(ImageSize.X, ImageSize.Y),
-		CameraMatrix,
+		ImageSize,
+		FocalLength,
+		ImageCenter,
 		DistortionCoefficients,
-		Rvecs,
-		Tvecs,
-		IntrinsicsDeviations,
-		ExtrinsicsDeviations,
-		ViewErrors,
-		cv::CALIB_USE_INTRINSIC_GUESS
+		CameraPoses,
+		PixelAspect,
+		SolverFlags
 	);
 
-	check(DistortionCoefficients.total() == 5);
-	check((CameraMatrix.rows == 3) && (CameraMatrix.cols == 3));
-
 	const TSharedPtr<FLensDistortionPointsRowData>& FirstRow = CalibrationRows[0];
 
-	// FZ inputs to LUT
-	OutFocus = FirstRow->CameraData.InputFocus;
-	OutZoom = FirstRow->CameraData.InputZoom;
-
-	// FocalLengthInfo
-	OutFocalLengthInfo.FxFy = FVector2D(
-		float(CameraMatrix.at<double>(0, 0) / ImageSize.X),
-		float(CameraMatrix.at<double>(1, 1) / ImageSize.Y)
-	);
+	OutLensModel = LensFile->LensInfo.LensModel;
+	OutDistortionInfo.Parameters = DistortionCoefficients;
 
-	// DistortionInfo
-	{
-		FSphericalDistortionParameters SphericalParameters;
-
-		SphericalParameters.K1 = DistortionCoefficients.at<double>(0);
-		SphericalParameters.K2 = DistortionCoefficients.at<double>(1);
-		SphericalParameters.P1 = DistortionCoefficients.at<double>(2);
-		SphericalParameters.P2 = DistortionCoefficients.at<double>(3);
-		SphericalParameters.K3 = DistortionCoefficients.at<double>(4);
+	OutFocalLengthInfo.FxFy = FVector2D(FocalLength / ImageSize);
 
-		USphericalLensModel::StaticClass()->GetDefaultObject<ULensModel>()->ToArray(
-			SphericalParameters,
-			OutDistortionInfo.Parameters
-		);
-	}
+	OutImageCenterInfo.PrincipalPoint = FVector2D(ImageCenter / ImageSize);
 
-	// ImageCenterInfo
-	OutImageCenterInfo.PrincipalPoint = FVector2D(
-		float(CameraMatrix.at<double>(0, 2) / ImageSize.X),
-		float(CameraMatrix.at<double>(1, 2) / ImageSize.Y)
-	);
-
-	// Lens Model
-	OutLensModel = USphericalLensModel::StaticClass();
+	// FZ inputs to LUT
+	OutFocus = FirstRow->CameraData.InputFocus;
+	OutZoom = FirstRow->CameraData.InputZoom;
 
 	return true;
-
-#else
-	OutErrorMessage = LOCTEXT("OpenCVRequired", "OpenCV is required");
-	return false;
-#endif //WITH_OPENCV
 }
 
 void UCameraLensDistortionAlgoPoints::OnDistortionSavedToLens()
@@ -848,18 +808,18 @@ TSharedRef<SWidget> UCameraLensDistortionAlgoPoints::BuildCurrentCalibratorPoint
 
 TSharedRef<SWidget> UCameraLensDistortionAlgoPoints::BuildCurrentFocalLengthWidget()
 {
-	return SNew(SNumericEntryBox<float>)
-		.Value_Lambda([&]() { return CurrentFocalLengthMM; })
-		.ToolTipText(LOCTEXT("CurrenFocalLength", "Current Focal Length (in mm)"))
-		.OnValueChanged_Lambda([&](float InValue)
+	return SNew(SNumericEntryBox<double>)
+		.Value_Lambda([&]() { return FocalLengthEstimate; })
+		.ToolTipText(LOCTEXT("CurrentFocalLength", "Current Focal Length (mm)"))
+		.OnValueChanged_Lambda([&](double InValue)
 		{
 			if (InValue < 1.0f)
 			{
 				InValue = 1.0f;
 			}
-			CurrentFocalLengthMM = InValue;
+			FocalLengthEstimate = InValue;
 		})
-		.OnValueCommitted_Lambda([&](float InValue, ETextCommit::Type CommitType)
+		.OnValueCommitted_Lambda([&](double InValue, ETextCommit::Type CommitType)
 		{
 			// Re-export the session data to update the new focal length value
 			ExportSessionData();
@@ -959,7 +919,7 @@ void UCameraLensDistortionAlgoPoints::ExportSessionData()
 		// Add all data to a json object that is needed to run this algorithm that is NOT part of a specific row
 		TSharedPtr<FJsonObject> JsonSessionData = MakeShared<FJsonObject>();
 
-		JsonSessionData->SetNumberField(LensDistortionPointsExportFields::CurrentFocalLength, CurrentFocalLengthMM);
+		JsonSessionData->SetNumberField(LensDistortionPointsExportFields::CurrentFocalLength, FocalLengthEstimate);
 		JsonSessionData->SetNumberField(LensDistortionPointsExportFields::Version, DATASET_VERSION);
 
 		// Export the session data to a .json file
@@ -993,7 +953,7 @@ void UCameraLensDistortionAlgoPoints::PreImportCalibrationData()
 void UCameraLensDistortionAlgoPoints::ImportSessionData(const TSharedRef<FJsonObject>& SessionDataObject)
 {
 	using namespace UE::CameraCalibration::Private;
-	SessionDataObject->TryGetNumberField(LensDistortionPointsExportFields::CurrentFocalLength, CurrentFocalLengthMM);
+	SessionDataObject->TryGetNumberField(LensDistortionPointsExportFields::CurrentFocalLength, FocalLengthEstimate);
 }
 
 int32 UCameraLensDistortionAlgoPoints::ImportCalibrationRow(const TSharedRef<FJsonObject>& CalibrationRowObject, const FImage& RowImage)
diff --git a/Engine/Plugins/VirtualProduction/CameraCalibration/Source/CameraCalibrationEditor/Private/Calibrators/CameraLensDistortionAlgoPoints.h b/Engine/Plugins/VirtualProduction/CameraCalibration/Source/CameraCalibrationEditor/Private/Calibrators/CameraLensDistortionAlgoPoints.h
--- a/Engine/Plugins/VirtualProduction/CameraCalibration/Source/CameraCalibrationEditor/Private/Calibrators/CameraLensDistortionAlgoPoints.h
+++ b/Engine/Plugins/VirtualProduction/CameraCalibration/Source/CameraCalibrationEditor/Private/Calibrators/CameraLensDistortionAlgoPoints.h
@@ -38,7 +38,7 @@ struct FLensDistortionPointsCalibratorPointData
 
 	// 2D location of the calibrator point
 	UPROPERTY()
-	FVector2D Point2d = FVector2D::ZeroVector;
+	FVector2f Point2d = FVector2f::ZeroVector;
 };
 
 /** Holds information of the camera pose for a given sample of a 2d-3d correlation */
@@ -188,7 +188,7 @@ protected:
 	uint32 CurrentPatternIndex = 0;
 
 	/** The value of the current focal length of the lens (in mm), used to set a guess for the camera intrinsics fx/fy */
-	float CurrentFocalLengthMM = 0.0f;
+	double FocalLengthEstimate = 0.0;
 
 protected:
 
diff --git a/Engine/Plugins/VirtualProduction/CameraCalibration/Source/CameraCalibrationEditor/Private/Calibrators/CameraNodalOffsetAlgoAruco.cpp b/Engine/Plugins/VirtualProduction/CameraCalibration/Source/CameraCalibrationEditor/Private/Calibrators/CameraNodalOffsetAlgoAruco.cpp
--- a/Engine/Plugins/VirtualProduction/CameraCalibration/Source/CameraCalibrationEditor/Private/Calibrators/CameraNodalOffsetAlgoAruco.cpp
+++ b/Engine/Plugins/VirtualProduction/CameraCalibration/Source/CameraCalibrationEditor/Private/Calibrators/CameraNodalOffsetAlgoAruco.cpp
@@ -7,26 +7,15 @@
 #include "AssetEditor/SSimulcamViewport.h"
 #include "CalibrationPointComponent.h"
 #include "CameraCalibrationUtils.h"
+#include "CameraCalibrationUtilsPrivate.h"
+#include "Engine/Texture2D.h"
 #include "Engine/TextureRenderTarget2D.h"
 #include "Misc/MessageDialog.h"
+#include "OpenCVHelper.h"
 #include "UI/CameraCalibrationWidgetHelpers.h"
 #include "Widgets/Input/SCheckBox.h"
 #include "Widgets/Views/SListView.h"
 
-#if WITH_OPENCV
-
-#include <vector>
-
-#include "OpenCVHelper.h"
-
-#include "PreOpenCVHeaders.h"
-#include "opencv2/aruco.hpp"
-#include "opencv2/calib3d.hpp"
-#include "opencv2/imgproc.hpp"
-#include "PostOpenCVHeaders.h"
-
-#endif //WITH_OPENCV
-
 #define LOCTEXT_NAMESPACE "CameraNodalOffsetAlgoAruco"
 
 
@@ -38,14 +27,12 @@ bool UCameraNodalOffsetAlgoAruco::OnViewportClicked(const FGeometry& MyGeometry,
 		return false;
 	}
 
-	{
-		const FText TitleError = LOCTEXT("CalibrationError", "CalibrationError");
-		FText ErrorMessage;
+	const FText TitleError = LOCTEXT("CalibrationError", "CalibrationError");
+	FText ErrorMessage;
 
-		if (!PopulatePoints(ErrorMessage))
-		{
-			FMessageDialog::Open(EAppMsgType::Ok, ErrorMessage, TitleError);
-		}
+	if (!PopulatePoints(ErrorMessage))
+	{
+		FMessageDialog::Open(EAppMsgType::Ok, ErrorMessage, TitleError);
 	}
 
 	return true;
@@ -91,17 +78,8 @@ TSharedRef<SWidget> UCameraNodalOffsetAlgoAruco::BuildUI()
 		;
 }
 
-#if WITH_OPENCV
-	#define BUILD_ARUCO_INFO(x) { FString(TEXT(#x)), cv::aruco::getPredefinedDictionary(cv::aruco::x) }
-#endif //WITH_OPENCV
-
 bool UCameraNodalOffsetAlgoAruco::PopulatePoints(FText& OutErrorMessage)
 {
-#if !WITH_OPENCV
-	OutErrorMessage = LOCTEXT("OpenCVRequired", "OpenCV required");
-	return false;
-#else
-
 	const FCameraCalibrationStepsController* StepsController;
 
 	if (!ensure(GetStepsControllerAndLensFile(&StepsController, nullptr)))
@@ -122,74 +100,13 @@ bool UCameraNodalOffsetAlgoAruco::PopulatePoints(FText& OutErrorMessage)
 		return false;
 	}
 
-	// Detect Aruco dictionary used by looking at the calibration point names
-
-	struct FArucoDictionaryInfo
-	{
-		FString Name;
-		cv::Ptr<cv::aruco::Dictionary> Dict;
-	} DictionaryInfo;
-
-	{
-		static TArray<FArucoDictionaryInfo> Dictionaries = 
-		{
-			BUILD_ARUCO_INFO(DICT_4X4_50),
-			BUILD_ARUCO_INFO(DICT_4X4_100),
-			BUILD_ARUCO_INFO(DICT_4X4_250),
-			BUILD_ARUCO_INFO(DICT_4X4_1000),
-			BUILD_ARUCO_INFO(DICT_5X5_50),
-			BUILD_ARUCO_INFO(DICT_5X5_100),
-			BUILD_ARUCO_INFO(DICT_5X5_250),
-			BUILD_ARUCO_INFO(DICT_5X5_1000),
-			BUILD_ARUCO_INFO(DICT_6X6_50),
-			BUILD_ARUCO_INFO(DICT_6X6_100),
-			BUILD_ARUCO_INFO(DICT_6X6_250),
-			BUILD_ARUCO_INFO(DICT_6X6_1000),
-			BUILD_ARUCO_INFO(DICT_7X7_50),
-			BUILD_ARUCO_INFO(DICT_7X7_100),
-			BUILD_ARUCO_INFO(DICT_7X7_250),
-			BUILD_ARUCO_INFO(DICT_7X7_1000),
-			BUILD_ARUCO_INFO(DICT_ARUCO_ORIGINAL),
-		};
-
-		for (const TWeakObjectPtr<const UCalibrationPointComponent>& CalibrationComponentPtr : ActiveCalibratorComponents)
-		{
-			if (const UCalibrationPointComponent* const CalibrationComponent = CalibrationComponentPtr.Get())
-			{
-				for (FArucoDictionaryInfo& Dictionary : Dictionaries)
-				{
-					if (CalibrationComponent->GetName().StartsWith(FString::Printf(TEXT("%s-"), *Dictionary.Name)))
-					{
-						DictionaryInfo = Dictionary;
-						break;
-					}
-
-					for (const TPair<FString, FVector>& SubPoint : CalibrationComponent->SubPoints)
-					{
-						if (SubPoint.Key.StartsWith(FString::Printf(TEXT("%s-"), *Dictionary.Name)))
-						{
-							DictionaryInfo = Dictionary;
-							break;
-						}
-					}
-
-					if (!DictionaryInfo.Dict.empty())
-					{
-						break;
-					}
-				}
 
-				if (!DictionaryInfo.Dict.empty())
-				{
-					break;
-				}
-			}
-		}
-	}
+	// Detect Aruco dictionary used by looking at the calibration point names
+	EArucoDictionary Dictionary = UE::CameraCalibration::Private::GetArucoDictionaryForCalibrator(Calibrator.Get());
 
-	if (DictionaryInfo.Dict.empty())
+	if (Dictionary == EArucoDictionary::None)
 	{
-		OutErrorMessage = LOCTEXT("InvalidArucoDictionary", 
+		OutErrorMessage = LOCTEXT("InvalidArucoDictionary",
 			"We could not find an appropriate Aruco dictionary for the selected Calibrator.\n"
 			"Make sure you have selected a Calibrator containing Calibration Point components\n"
 			"with names per the naming convention described in the info section of this algo"
@@ -215,25 +132,10 @@ bool UCameraNodalOffsetAlgoAruco::PopulatePoints(FText& OutErrorMessage)
 		return false;
 	}
 
-	// Create OpenCV Mat with those pixels
-	cv::Mat CvFrame(cv::Size(Size.X, Size.Y), CV_8UC4, Pixels.GetData());
+	TArray<FArucoMarker> IdentifiedMarkers;
+	bool bResult = FOpenCVHelper::IdentifyArucoMarkers(Pixels, Size, Dictionary, IdentifiedMarkers);
 
-	// Convert to Gray
-	cv::Mat CvGray;
-	cv::cvtColor(CvFrame, CvGray, cv::COLOR_RGBA2GRAY);
-	
-	// Run the Aruco parser
-	
-	std::vector<int> MarkerIds;
-	std::vector<std::vector<cv::Point2f>> MarkerCorners;
-	cv::Ptr<cv::aruco::DetectorParameters> DetectorParameters = cv::aruco::DetectorParameters::create();
-	DetectorParameters->cornerRefinementMethod = cv::aruco::CORNER_REFINE_SUBPIX;
-
-	cv::aruco::detectMarkers(CvGray, DictionaryInfo.Dict, MarkerCorners, MarkerIds, DetectorParameters);
-
-	check((Size.X > 0) && (Size.Y > 0));
-
-	if (MarkerCorners.empty())
+	if (!bResult || IdentifiedMarkers.IsEmpty())
 	{
 		OutErrorMessage = LOCTEXT("NoMarkerCornersFound", "No Aruco markers were detected in the media image, "
 			"or they did not match the aruco dictionary of any of the selected calibrator components."
@@ -244,21 +146,15 @@ bool UCameraNodalOffsetAlgoAruco::PopulatePoints(FText& OutErrorMessage)
 	// Show the detection to the user
 	if (bShouldShowDetectionWindow)
 	{
-		// We need to convert to RGB temporarily because cv::aruco::drawDetectedMarkers only supports gray or RGB
-		cv::cvtColor(CvFrame, CvFrame, cv::COLOR_RGBA2RGB);
-		cv::aruco::drawDetectedMarkers(CvFrame, MarkerCorners, MarkerIds);
-		cv::cvtColor(CvFrame, CvFrame, cv::COLOR_RGB2RGBA);
-
-		FCameraCalibrationWidgetHelpers::DisplayTextureInWindowAlmostFullScreen(
-			FOpenCVHelper::TextureFromCvMat(CvFrame),
-			LOCTEXT("ArucoMarkerDetection", "ArucoMarkerDetection")
-		);
-	}
+		UTexture2D* DebugTexture = UTexture2D::CreateTransient(Size.X, Size.Y, EPixelFormat::PF_B8G8R8A8);
+		UE::CameraCalibration::Private::SetTextureData(DebugTexture, Pixels);
+		FOpenCVHelper::DrawArucoMarkers(IdentifiedMarkers, DebugTexture);
 
-	check(MarkerIds.size() == MarkerCorners.size());
+		FCameraCalibrationWidgetHelpers::DisplayTextureInWindowAlmostFullScreen(DebugTexture, LOCTEXT("ArucoMarkerDetection", "ArucoMarkerDetection"));
+	}
 
 	// Add the detected markers that match calibration points to the table
-	for(int32 MarkerIdx=0; MarkerIdx < MarkerCorners.size(); ++MarkerIdx)
+	for(int32 MarkerIdx=0; MarkerIdx < IdentifiedMarkers.Num(); ++MarkerIdx)
 	{
 		// Naming convention for calibration points:
 		//     DICTIONARY-ID-CORNER
@@ -268,9 +164,9 @@ bool UCameraNodalOffsetAlgoAruco::PopulatePoints(FText& OutErrorMessage)
 		//     DICT_6X6_250-23-TL
 		//     DICT_6X6_250-23-TR
 
-		static const TArray<FString> CornerNames = { TEXT("TL"), TEXT("TR"), TEXT("BR"), TEXT("BL") };
+		FArucoMarker Marker = IdentifiedMarkers[MarkerIdx];
 
-		check(MarkerCorners[MarkerIdx].size() == 4); // Successful detection should always return the 4 corners of each marker.
+		static const TArray<FString> CornerNames = { TEXT("TL"), TEXT("TR"), TEXT("BR"), TEXT("BL") };
 
 		// Export the latest session data
 		ExportSessionData();
@@ -279,7 +175,8 @@ bool UCameraNodalOffsetAlgoAruco::PopulatePoints(FText& OutErrorMessage)
 		for (int32 CornerIdx = 0; CornerIdx < 4; ++CornerIdx)
 		{
 			// Build calibrator point name based on the detected marker
-			const FString MarkerName = FString::Printf(TEXT("%s-%d-%s"), *DictionaryInfo.Name, MarkerIds[MarkerIdx], *CornerNames[CornerIdx]); //-V557
+			const FString DictionaryName = UE::CameraCalibration::Private::GetArucoDictionaryName(Dictionary);
+			const FString MarkerName = FString::Printf(TEXT("%s-%d-%s"), *DictionaryName, Marker.MarkedId, *CornerNames[CornerIdx]); //-V557
 
 			FNodalOffsetPointsCalibratorPointData PointCache;
 
@@ -295,8 +192,7 @@ bool UCameraNodalOffsetAlgoAruco::PopulatePoints(FText& OutErrorMessage)
 			const uint32 RowIndex = NodalOffsetTool->AdvanceSessionRowIndex();
 
 			Row->Index = RowIndex;
-			Row->Point2D.X = float(MarkerCorners[MarkerIdx][CornerIdx].x) / Size.X;
-			Row->Point2D.Y = float(MarkerCorners[MarkerIdx][CornerIdx].y) / Size.Y;
+			Row->Point2D = Marker.Corners[CornerIdx] / Size;
 
 			Row->CalibratorPointData.Location = PointCache.Location;
 			Row->CalibratorPointData.Name = MarkerName;
@@ -329,8 +225,6 @@ bool UCameraNodalOffsetAlgoAruco::PopulatePoints(FText& OutErrorMessage)
 	}
 
 	return true;
-
-#endif //WITH_OPENCV
 }
 
 TSharedRef<SWidget> UCameraNodalOffsetAlgoAruco::BuildShowDetectionWidget()
diff --git a/Engine/Plugins/VirtualProduction/CameraCalibration/Source/CameraCalibrationEditor/Private/Calibrators/CameraNodalOffsetAlgoCheckerboard.cpp b/Engine/Plugins/VirtualProduction/CameraCalibration/Source/CameraCalibrationEditor/Private/Calibrators/CameraNodalOffsetAlgoCheckerboard.cpp
--- a/Engine/Plugins/VirtualProduction/CameraCalibration/Source/CameraCalibrationEditor/Private/Calibrators/CameraNodalOffsetAlgoCheckerboard.cpp
+++ b/Engine/Plugins/VirtualProduction/CameraCalibration/Source/CameraCalibrationEditor/Private/Calibrators/CameraNodalOffsetAlgoCheckerboard.cpp
@@ -8,25 +8,12 @@
 #include "Engine/TextureRenderTarget2D.h"
 #include "EngineUtils.h"
 #include "Misc/MessageDialog.h"
+#include "OpenCVHelper.h"
 #include "UI/CameraCalibrationWidgetHelpers.h"
 #include "UI/SFilterableActorPicker.h"
 #include "Widgets/Views/SListView.h"
 
 
-#if WITH_OPENCV
-
-#include <vector>
-#include <algorithm>
-
-#include "OpenCVHelper.h"
-#include "PreOpenCVHeaders.h"
-#include "opencv2/core/types.hpp"
-#include "opencv2/calib3d.hpp"
-#include "opencv2/imgproc.hpp"
-#include "PostOpenCVHeaders.h"
-
-#endif //WITH_OPENCV
-
 #define LOCTEXT_NAMESPACE "CameraNodalOffsetAlgoCheckerboard"
 
 
@@ -102,11 +89,6 @@ TSharedRef<SWidget> UCameraNodalOffsetAlgoCheckerboard::BuildCheckerboardPickerW
 
 bool UCameraNodalOffsetAlgoCheckerboard::PopulatePoints(FText& OutErrorMessage)
 {
-#if !WITH_OPENCV
-	OutErrorMessage = LOCTEXT("OpenCVRequired", "OpenCV required");
-	return false;
-#else
-
 	const FCameraCalibrationStepsController* StepsController;
 
 	if (!ensure(GetStepsControllerAndLensFile(&StepsController, nullptr)))
@@ -142,66 +124,23 @@ bool UCameraNodalOffsetAlgoCheckerboard::PopulatePoints(FText& OutErrorMessage)
 		return false;
 	}
 
-	// Create OpenCV Mat with those pixels
-	cv::Mat CvFrame(cv::Size(Size.X, Size.Y), CV_8UC4, Pixels.GetData());
-
-	// Convert to Gray
-	cv::Mat CvGray;
-	cv::cvtColor(CvFrame, CvGray, cv::COLOR_RGBA2GRAY);
-
-	// Populate the 3d/2d correlation points
-
 	const ACameraCalibrationCheckerboard* Checkerboard = Cast<ACameraCalibrationCheckerboard>(Calibrator.Get());
 	check(Checkerboard); // This should be ensured by the picker.
 
-	std::vector<cv::Point2f> Points2d;
-
-	// Identify checkerboard
-	{
-		cv::Size CheckerboardSize(Checkerboard->NumCornerCols, Checkerboard->NumCornerRows);
+	const FIntPoint CheckerboardDimensions = FIntPoint(Checkerboard->NumCornerCols, Checkerboard->NumCornerRows);
 
-		std::vector<cv::Point2f> Corners;
+	TArray<FVector2f> Corners;
+	const bool bCornersFound = FOpenCVHelper::IdentifyCheckerboard(Pixels, Size, CheckerboardDimensions, Corners);
 
-		const bool bCornersFound = cv::findChessboardCorners(
-			CvGray,
-			CheckerboardSize,
-			Corners,
-			cv::CALIB_CB_ADAPTIVE_THRESH | cv::CALIB_CB_NORMALIZE_IMAGE
+	if (!bCornersFound || Corners.IsEmpty())
+	{
+		OutErrorMessage = FText::FromString(FString::Printf(TEXT(
+			"Could not identify the expected checkerboard points of interest. "
+			"The expected checkerboard has %dx%d inner corners."),
+			Checkerboard->NumCornerCols, Checkerboard->NumCornerRows)
 		);
 
-		if (!bCornersFound)
-		{
-			OutErrorMessage = FText::FromString(FString::Printf(TEXT(
-				"Could not identify the expected checkerboard points of interest. "
-				"The expected checkerboard has %dx%d inner corners."),
-				Checkerboard->NumCornerCols, Checkerboard->NumCornerRows)
-			);
-
-			return false;
-		}
-
-		// CV_TERMCRIT_EPS will stop the search when the error is under the given epsilon.
-		// CV_TERMCRIT_ITER will stop after the specified number of iterations regardless of epsilon.
-		cv::TermCriteria Criteria(cv::TermCriteria::Type::EPS | cv::TermCriteria::Type::COUNT, 30, 0.001);
-		cv::cornerSubPix(CvGray, Corners, cv::Size(11, 11), cv::Size(-1, -1), Criteria);
-
-		Points2d.reserve(Corners.size());
-
-		if (!Corners.empty())
-		{
-			for (cv::Point2f& Corner : Corners)
-			{
-				Points2d.push_back(cv::Point2f(Corner.x, Corner.y));
-			}
-		}
-	}
-
-	check(Points2d.size() == Checkerboard->NumCornerRows * Checkerboard->NumCornerCols);
-
-	// Force first one to be top-left
-	if (Points2d.front().y > Points2d.back().y) 
-	{
-		std::reverse(Points2d.begin(), Points2d.end());
+		return false;
 	}
 
 	// Export the latest session data
@@ -220,8 +159,7 @@ bool UCameraNodalOffsetAlgoCheckerboard::PopulatePoints(FText& OutErrorMessage)
 			const uint32 RowIndex = NodalOffsetTool->AdvanceSessionRowIndex();
 
 			Row->Index = RowIndex;
-			Row->Point2D.X = float(Points2d[PointIdx].x) / Size.X;
-			Row->Point2D.Y = float(Points2d[PointIdx].y) / Size.Y;
+			Row->Point2D = Corners[PointIdx] / Size;
 
 			FTransform LocalPoint3d;
 			LocalPoint3d.SetLocation(Checkerboard->SquareSideLength * FVector(0, ColIdx, Checkerboard->NumCornerRows - RowIdx - 1));
@@ -257,8 +195,6 @@ bool UCameraNodalOffsetAlgoCheckerboard::PopulatePoints(FText& OutErrorMessage)
 	}
 
 	return true;
-
-#endif //WITH_OPENCV
 }
 
 AActor* UCameraNodalOffsetAlgoCheckerboard::FindFirstCalibrator() const
diff --git a/Engine/Plugins/VirtualProduction/CameraCalibration/Source/CameraCalibrationEditor/Private/Calibrators/CameraNodalOffsetAlgoOpticalAxis.cpp b/Engine/Plugins/VirtualProduction/CameraCalibration/Source/CameraCalibrationEditor/Private/Calibrators/CameraNodalOffsetAlgoOpticalAxis.cpp
--- a/Engine/Plugins/VirtualProduction/CameraCalibration/Source/CameraCalibrationEditor/Private/Calibrators/CameraNodalOffsetAlgoOpticalAxis.cpp
+++ b/Engine/Plugins/VirtualProduction/CameraCalibration/Source/CameraCalibrationEditor/Private/Calibrators/CameraNodalOffsetAlgoOpticalAxis.cpp
@@ -19,15 +19,6 @@
 #include "Widgets/Input/SComboBox.h"
 #include "Widgets/Input/SNumericEntryBox.h"
 
-#include <vector>
-
-#if WITH_OPENCV
-#include "PreOpenCVHeaders.h"
-#include "opencv2/calib3d.hpp"
-#include "opencv2/imgproc.hpp"
-#include "PostOpenCVHeaders.h"
-#endif
-
 
 #define LOCTEXT_NAMESPACE "CameraNodalOffsetAlgoOpticalAxis"
 
@@ -231,26 +222,20 @@ bool UCameraNodalOffsetAlgoOpticalAxis::GetNodalOffset(FNodalPointOffset& OutNod
 	OutFocus = CachedFocus;
 	OutZoom = CachedZoom;
 
-#if WITH_OPENCV
 	// Gather the 3D points from each of the calibration rows
-	std::vector<cv::Point3f> Points3d;
-	Points3d.reserve(CalibrationRows.Num());
+	const int32 NumPoints = CalibrationRows.Num();
+
+	TArray<FVector> Points3d;
+	Points3d.Reserve(NumPoints);
 
 	for (const TSharedPtr<FNodalOffsetPointsRowData>& Row : CalibrationRows)
 	{
-		Points3d.push_back(cv::Point3f(
-			Row->CalibratorPointData.Location.X,
-			Row->CalibratorPointData.Location.Y,
-			Row->CalibratorPointData.Location.Z
-		));
+		Points3d.Add(Row->CalibratorPointData.Location);
 	}
 
-	// Find a best fit line between the 3D points, producing a line (the optical axis) and a point on that line
-	cv::Vec6f LineAndPoint;
-	cv::fitLine(Points3d, LineAndPoint, cv::DIST_L2, 0, 0.01, 0.01);
-
-	OpticalAxis = FVector(LineAndPoint[0], LineAndPoint[1], LineAndPoint[2]);
-	PointAlongAxis = FVector(LineAndPoint[3], LineAndPoint[4], LineAndPoint[5]);
+	FVector FitLine;
+	FOpenCVHelper::FitLine3D(Points3d, FitLine, PointAlongAxis);
+	OpticalAxis = FitLine;
 
 	// Find a look at rotation for the camera
 	const FVector Point0 = PointAlongAxis + (0.0f * OpticalAxis.GetValue());
@@ -283,14 +268,8 @@ bool UCameraNodalOffsetAlgoOpticalAxis::GetNodalOffset(FNodalPointOffset& OutNod
 	// which would indicate that they were in roughly a straight line. 
 	const FLensDistortionState DistortionState = DistortionHandler->GetCurrentDistortionState();
 
-	cv::Mat CameraMatrix(3, 3, cv::DataType<double>::type);
-	cv::setIdentity(CameraMatrix);
-
-	CameraMatrix.at<double>(0, 0) = DistortionState.FocalLengthInfo.FxFy.X;
-	CameraMatrix.at<double>(1, 1) = DistortionState.FocalLengthInfo.FxFy.Y;
-
-	CameraMatrix.at<double>(0, 2) = DistortionState.ImageCenter.PrincipalPoint.X;
-	CameraMatrix.at<double>(1, 2) = DistortionState.ImageCenter.PrincipalPoint.Y;
+	const FVector2D FocalLength = DistortionState.FocalLengthInfo.FxFy;
+	const FVector2D ImageCenter = DistortionState.ImageCenter.PrincipalPoint;
 
 	// Compute the optimal camera pose using the new nodal offset and convert it to opencv's coordinate system
 	FTransform NewNodalOffsetTransform;
@@ -299,35 +278,15 @@ bool UCameraNodalOffsetAlgoOpticalAxis::GetNodalOffset(FNodalPointOffset& OutNod
 
 	FTransform OptimalCameraPose = NewNodalOffsetTransform * CachedNodalOffset.Inverse() * CameraPose;
 
-	FOpenCVHelper::ConvertUnrealToOpenCV(OptimalCameraPose);
-
-	// Transform the 3D points to opencv's coordinate system
-	for (int32 Index = 0; Index < Points3d.size(); ++Index)
-	{
-		FTransform PointTransform = FTransform::Identity;
-		PointTransform.SetLocation(FVector(Points3d[Index].x, Points3d[Index].y, Points3d[Index].z));
-		FOpenCVHelper::ConvertUnrealToOpenCV(PointTransform);
-		Points3d[Index] = cv::Point3f(PointTransform.GetLocation().X, PointTransform.GetLocation().Y, PointTransform.GetLocation().Z);
-	}
-
 	// Initialize all of the corresponding 2D points to be exactly at the principal point of the lens
-	std::vector<cv::Point2f> Points2d;
-	for (int32 Index = 0; Index < Points3d.size(); ++Index)
-	{
-		Points2d.push_back(cv::Point2f(DistortionState.ImageCenter.PrincipalPoint.X, DistortionState.ImageCenter.PrincipalPoint.Y));
-	}
+	TArray<FVector2f> ImagePoints;
+	ImagePoints.Init(FVector2f(ImageCenter.X, ImageCenter.Y), NumPoints);
 
-	OutError = FOpenCVHelper::ComputeReprojectionError(OptimalCameraPose, CameraMatrix, Points3d, Points2d);
+	OutError = FOpenCVHelper::ComputeReprojectionError(Points3d, ImagePoints, FocalLength, ImageCenter, OptimalCameraPose);
 
 	bFocusZoomWarningIssued = false;
 
 	return true;
-#else
-	{
-		OutErrorMessage = LOCTEXT("OpenCVRequired", "OpenCV is required");
-		return false;
-	}
-#endif // WITH_OPENCV
 }
 
 UMaterialInterface* UCameraNodalOffsetAlgoOpticalAxis::GetOverlayMaterial() const
diff --git a/Engine/Plugins/VirtualProduction/CameraCalibration/Source/CameraCalibrationEditor/Private/Calibrators/CameraNodalOffsetAlgoPoints.cpp b/Engine/Plugins/VirtualProduction/CameraCalibration/Source/CameraCalibrationEditor/Private/Calibrators/CameraNodalOffsetAlgoPoints.cpp
--- a/Engine/Plugins/VirtualProduction/CameraCalibration/Source/CameraCalibrationEditor/Private/Calibrators/CameraNodalOffsetAlgoPoints.cpp
+++ b/Engine/Plugins/VirtualProduction/CameraCalibration/Source/CameraCalibrationEditor/Private/Calibrators/CameraNodalOffsetAlgoPoints.cpp
@@ -7,6 +7,7 @@
 #include "Camera/CameraActor.h"
 #include "Camera/CameraComponent.h"
 #include "CameraCalibrationEditorLog.h"
+#include "CameraCalibrationSolver.h"
 #include "CameraCalibrationUtils.h"
 #include "CineCameraComponent.h"
 #include "DistortionRenderingUtils.h"
@@ -35,21 +36,8 @@
 #include "Widgets/Views/SListView.h"
 #include "Widgets/SWidget.h"
 
-#include <vector>
-
-#if WITH_OPENCV
-#include "PreOpenCVHeaders.h"
-#include "opencv2/calib3d.hpp"
-#include "PostOpenCVHeaders.h"
-#endif
-
 #define LOCTEXT_NAMESPACE "CameraNodalOffsetAlgoPoints"
 
-#if WITH_EDITOR
-static TAutoConsoleVariable<float> CVarRotationStepValue(TEXT("CameraCalibration.RotationStepValue"), 0.05, TEXT("The value of the initial step size to use when finding an optimal nodal offset rotation that minimizes reprojection error."));
-static TAutoConsoleVariable<float> CVarLocationStepValue(TEXT("CameraCalibration.LocationStepValue"), 0.5, TEXT("The value of the initial step size to use when finding an optimal nodal offset location that minimizes reprojection error."));
-#endif
-
 const int UCameraNodalOffsetAlgoPoints::DATASET_VERSION = 1;
 
 // String constants for import/export
@@ -205,122 +193,8 @@ namespace CameraNodalOffsetAlgoPoints
 	}
 };
 
-#if WITH_OPENCV
-class FMinReprojectionErrorSolver : public cv::MinProblemSolver::Function
-{
-public:
-	struct FSingleViewPoseAndPoints
-	{
-		FTransform CameraPose;
-
-		std::vector<cv::Point3f> Points3d;
-		std::vector<cv::Point2f> Points2d;
-	};
-
-	//~ Begin cv::MinProblemSolver::Function interface
-	int getDims() const
-	{
-		return 7; // 4 doubles for the rotation quaternion + 3 doubles for the location vector
-	}
-
-	double calc(const double* x) const
-	{
-		// Convert the input data (7 doubles) to an FQuat and FVector
-		const FQuat Rotation = FQuat(x[0], x[1], x[2], x[3]).GetNormalized();
-		const FVector Location = FVector(x[4], x[5], x[6]);
-
-		UE_LOG(LogCameraCalibrationEditor, VeryVerbose, TEXT("Nodal Offset Candidate:  Rotation: (%lf, %lf, %lf, %lf)  Location: (%lf, %lf, %lf)"),
-			Rotation.X, Rotation.Y, Rotation.Z, Rotation.W, Location.X, Location.Y, Location.Z);
-
-		FTransform NodalOffsetCandidate;
-
-		// As a result of the way that the downhill solver nudges the input data on each iteration, it is important to normalize the rotation
-		NodalOffsetCandidate.SetRotation(Rotation);
-		NodalOffsetCandidate.SetLocation(Location);
-
-		double ReprojectionErrorTotal = 0.0;
-		int32 NumTotalPoints = 0;
-
-		// For each camera view
-		for (const FSingleViewPoseAndPoints& CameraView : CameraViews)
-		{
-			// Compute the optimal camera pose using the nodal offset candidate for this iteration and the tracked camera pose 
-			// Any existing nodal offset that was applied to the camera is also factored in
-			FTransform OptimalCameraPose = NodalOffsetCandidate * ExistingNodalOffsetInverse * CameraView.CameraPose;
-
-			// Convert the optimal camera pose to opencv's coordinate system
-			FOpenCVHelper::ConvertUnrealToOpenCV(OptimalCameraPose);
-
-			// Compute the reprojection error for this view and add it to the running total
-			double ViewError = FOpenCVHelper::ComputeReprojectionError(OptimalCameraPose, CameraMatrix, CameraView.Points3d, CameraView.Points2d);
-			ReprojectionErrorTotal += ViewError;
-
-			NumTotalPoints += CameraView.Points2d.size();
-		}
-
-		double RootMeanSquareError = FMath::Sqrt(ReprojectionErrorTotal / NumTotalPoints);
-
-		UE_LOG(LogCameraCalibrationEditor, VeryVerbose, TEXT("Reprojection Error: %lf"), RootMeanSquareError);
-
-		return RootMeanSquareError;
-	}
-	//~ End cv::MinProblemSolver::Function interface
-
-	void SetCameraMatrix(const cv::Mat InCameraMatrix) 
-	{ 
-		CameraMatrix = InCameraMatrix; 
-	}
-
-	void SetExistingNodalOffsetInverse(FTransform InOffsetInverse) 
-	{ 
-		ExistingNodalOffsetInverse = InOffsetInverse; 
-	}
-
-	void AddCameraViewAndPoints(FTransform InCameraPose, std::vector<cv::Point3f> InPoints3d, std::vector<cv::Point2f> InPoints2d)
-	{
-		FSingleViewPoseAndPoints& NewCameraView = CameraViews.Add_GetRef(FSingleViewPoseAndPoints());
-
-		NewCameraView.CameraPose = InCameraPose;
-		NewCameraView.Points3d = MoveTemp(InPoints3d);
-		NewCameraView.Points2d = MoveTemp(InPoints2d);
-	}
-
-private:
-	cv::Mat CameraMatrix;
-
-	FTransform ExistingNodalOffsetInverse;
-
-	TArray<FSingleViewPoseAndPoints> CameraViews;
-};
-#endif
-
 double UCameraNodalOffsetAlgoPoints::MinimizeReprojectionError(FTransform& InOutNodalOffset, const TArray<TSharedPtr<TArray<TSharedPtr<FNodalOffsetPointsRowData>>>>& SamePoseRowGroups) const
 {
-#if WITH_OPENCV
-	cv::Ptr<cv::DownhillSolver> Solver = cv::DownhillSolver::create();
-
-	cv::Ptr<FMinReprojectionErrorSolver> SolverFunction = cv::makePtr<FMinReprojectionErrorSolver>();
-	Solver->setFunction(SolverFunction);
-
-	const FQuat InitialRotation = InOutNodalOffset.GetRotation();
-	const FVector InitialLocation = InOutNodalOffset.GetLocation();
-
-	cv::Mat WorkingSolution(1, 7, cv::DataType<double>::type);
-	WorkingSolution.at<double>(0, 0) = InitialRotation.X;
-	WorkingSolution.at<double>(0, 1) = InitialRotation.Y;
-	WorkingSolution.at<double>(0, 2) = InitialRotation.Z;
-	WorkingSolution.at<double>(0, 3) = InitialRotation.W;
-	WorkingSolution.at<double>(0, 4) = InitialLocation.X;
-	WorkingSolution.at<double>(0, 5) = InitialLocation.Y;
-	WorkingSolution.at<double>(0, 6) = InitialLocation.Z;
-
-	// NOTE: These step sizes may need further testing and refinement, but tests so far have shown them to be decent
-	const double RotationStep = CVarRotationStepValue.GetValueOnAnyThread();
-	const double LocationStep = CVarLocationStepValue.GetValueOnAnyThread();
-	cv::Mat Step = (cv::Mat_<double>(7, 1) << RotationStep, RotationStep, RotationStep, RotationStep, LocationStep, LocationStep, LocationStep);
-
-	Solver->setInitStep(Step);
-
 	const FCameraCalibrationStepsController* StepsController;
 	const ULensFile* LensFile;
 
@@ -337,18 +211,8 @@ double UCameraNodalOffsetAlgoPoints::MinimizeReprojectionError(FTransform& InOut
 
 	const FLensDistortionState DistortionState = DistortionHandler->GetCurrentDistortionState();
 
-	// Initialize the camera matrix that will be used in each call to projectPoints()
-	cv::Mat CameraMatrix(3, 3, cv::DataType<double>::type);
-	cv::setIdentity(CameraMatrix);
-
-	CameraMatrix.at<double>(0, 0) = DistortionState.FocalLengthInfo.FxFy.X;
-	CameraMatrix.at<double>(1, 1) = DistortionState.FocalLengthInfo.FxFy.Y;
-
-	// Note that the 2D points that will be compared against have been undistorted and center shift has already been accounted for
-	CameraMatrix.at<double>(0, 2) = 0.5;
-	CameraMatrix.at<double>(1, 2) = 0.5;
-
-	SolverFunction->SetCameraMatrix(CameraMatrix);
+	const FVector2D FocalLength = DistortionState.FocalLengthInfo.FxFy;
+	const FVector2D ImageCenter = FVector2D(0.5, 0.5);
 
 	if (!ensureMsgf((CalibrationRows.Num() > 0), TEXT("Not enough calibration rows")))
 	{
@@ -372,14 +236,23 @@ double UCameraNodalOffsetAlgoPoints::MinimizeReprojectionError(FTransform& InOut
 		}
 	}
 
-	SolverFunction->SetExistingNodalOffsetInverse(ExistingOffset.Inverse());
-
 	if (!ensureMsgf(FirstRow->bUndistortedIsValid, TEXT("This method operates on undistorted 2D points. Call UndistortCalibrationRowPoints() prior to calling this method")))
 	{
 		return -1.0;
 	}
 
-	for (int32 PoseGroupIndex = 0; PoseGroupIndex < SamePoseRowGroups.Num(); ++PoseGroupIndex)
+	const int32 NumGroups = SamePoseRowGroups.Num();
+
+	TArray<TArray<FVector>> ObjectPoints;
+	ObjectPoints.Reserve(NumGroups);
+
+	TArray<TArray<FVector2f>> ImagePoints;
+	ImagePoints.Reserve(NumGroups);
+
+	TArray<FTransform> CameraPoses;
+	CameraPoses.Reserve(NumGroups);
+
+	for (int32 PoseGroupIndex = 0; PoseGroupIndex < NumGroups; ++PoseGroupIndex)
 	{
 		const TArray<TSharedPtr<FNodalOffsetPointsRowData>>& PoseGroup = *(SamePoseRowGroups[PoseGroupIndex]);
 
@@ -389,148 +262,27 @@ double UCameraNodalOffsetAlgoPoints::MinimizeReprojectionError(FTransform& InOut
 		}
 
 		const TSharedPtr<FNodalOffsetPointsRowData>& FirstRowInPoseGroup = PoseGroup[0];
+		const FTransform CameraPose = FirstRowInPoseGroup->CameraData.Pose;
 
-		FTransform CameraPose = FirstRowInPoseGroup->CameraData.Pose;
+		// Undo the existing nodal offset transformation
+		CameraPoses.Add(ExistingOffset.Inverse() * CameraPose);
 
-		std::vector<cv::Point3f> Points3d;
-		std::vector<cv::Point2f> Points2d;
-		Points3d.reserve(PoseGroup.Num());
-		Points2d.reserve(PoseGroup.Num());
+		TArray<FVector> Points3d;
+		TArray<FVector2f> Points2d;
+		Points3d.Reserve(PoseGroup.Num());
+		Points2d.Reserve(PoseGroup.Num());
 
 		for (const TSharedPtr<FNodalOffsetPointsRowData>& Row : PoseGroup)
 		{
-			// Convert from UE coordinates to OpenCV coordinates
-			FTransform Transform;
-			Transform.SetIdentity();
-			Transform.SetLocation(Row->CalibratorPointData.Location);
-
-			FOpenCVHelper::ConvertUnrealToOpenCV(Transform);
-
-			// Calibrator 3d points
-			Points3d.push_back(cv::Point3f(
-				Transform.GetLocation().X,
-				Transform.GetLocation().Y,
-				Transform.GetLocation().Z));
-
-			Points2d.push_back(cv::Point2f(
-				Row->UndistortedPoint2D.X,
-				Row->UndistortedPoint2D.Y));
+			Points3d.Add(Row->CalibratorPointData.Location);
+			Points2d.Add(Row->UndistortedPoint2D);
 		}
 
-		SolverFunction->AddCameraViewAndPoints(CameraPose, Points3d, Points2d);
-	}
-
-	double Error = Solver->minimize(WorkingSolution);
-
-	const FQuat FinalRotation = FQuat(WorkingSolution.at<double>(0, 0), WorkingSolution.at<double>(0, 1), WorkingSolution.at<double>(0, 2), WorkingSolution.at<double>(0, 3)).GetNormalized();
-	const FVector FinalLocation = FVector(WorkingSolution.at<double>(0, 4), WorkingSolution.at<double>(0, 5), WorkingSolution.at<double>(0, 6));
-
-	InOutNodalOffset.SetRotation(FinalRotation);
-	InOutNodalOffset.SetLocation(FinalLocation);
-
-	return Error;
-#else
-	return -1.0;
-#endif
-}
-
-double UCameraNodalOffsetAlgoPoints::ComputeReprojectionError(const FTransform& NodalOffset, const TArray<TSharedPtr<FNodalOffsetPointsRowData>>& PoseGroup) const
-{
-#if WITH_OPENCV
-	const FCameraCalibrationStepsController* StepsController;
-	const ULensFile* LensFile;
-
-	if (!ensure(GetStepsControllerAndLensFile(&StepsController, &LensFile)))
-	{
-		return -1.0;
-	}
-
-	const ULensDistortionModelHandlerBase* DistortionHandler = StepsController->GetDistortionHandler();
-	if (!DistortionHandler)
-	{
-		return -1.0;
-	}
-
-	std::vector<cv::Point3f> Points3d;
-	std::vector<cv::Point2f> Points2d;
-	Points3d.reserve(PoseGroup.Num());
-	Points2d.reserve(PoseGroup.Num());
-
-	if (!ensureMsgf((PoseGroup.Num() > 0), TEXT("Not enough calibration rows")))
-	{
-		return -1.0;
-	}
-
-	const TSharedPtr<FNodalOffsetPointsRowData>& FirstRow = PoseGroup[0];
-
-	if (!ensureMsgf(FirstRow->bUndistortedIsValid, TEXT("This method operates on undistorted 2D points. Call UndistortCalibrationRowPoints() prior to calling this method")))
-	{
-		return -1.0;
-	}
-
-	for (const TSharedPtr<FNodalOffsetPointsRowData>& Row : PoseGroup)
-	{
-		// Convert from UE coordinates to OpenCV coordinates
-
-		FTransform Transform;
-		Transform.SetIdentity();
-		Transform.SetLocation(Row->CalibratorPointData.Location);
-
-		FOpenCVHelper::ConvertUnrealToOpenCV(Transform);
-
-		// Calibrator 3d points
-		Points3d.push_back(cv::Point3f(
-			Transform.GetLocation().X,
-			Transform.GetLocation().Y,
-			Transform.GetLocation().Z));
-
-		Points2d.push_back(cv::Point2f(
-			Row->UndistortedPoint2D.X,
-			Row->UndistortedPoint2D.Y));
+		ObjectPoints.Add(Points3d);
+		ImagePoints.Add(Points2d);
 	}
 
-	// Compute the optimal camera pose using the nodal offset candidate for this iteration and the tracked camera pose 
-	// Any existing nodal offset that was applied to the camera is also factored in
-	FTransform CameraPose = FirstRow->CameraData.Pose;
-
-	FTransform ExistingOffset = FTransform::Identity;
-
-	if (FirstRow->CameraData.bWasNodalOffsetApplied)
-	{
-		FNodalPointOffset NodalPointOffset;
-		const float Focus = FirstRow->CameraData.InputFocus;
-		const float Zoom = FirstRow->CameraData.InputZoom;
-
-		if (LensFile->EvaluateNodalPointOffset(Focus, Zoom, NodalPointOffset))
-		{
-			ExistingOffset.SetTranslation(NodalPointOffset.LocationOffset);
-			ExistingOffset.SetRotation(NodalPointOffset.RotationOffset);
-		}
-	}
-
-	// Initialize the camera matrix that will be used in each call to projectPoints()
-	FLensDistortionState DistortionState = DistortionHandler->GetCurrentDistortionState();
-
-	cv::Mat CameraMatrix(3, 3, cv::DataType<double>::type);
-	cv::setIdentity(CameraMatrix);
-
-	CameraMatrix.at<double>(0, 0) = DistortionState.FocalLengthInfo.FxFy.X;
-	CameraMatrix.at<double>(1, 1) = DistortionState.FocalLengthInfo.FxFy.Y;
-
-	// Note that the 2D points that will be compared against have been undistorted and center shift has already been accounted for
-	CameraMatrix.at<double>(0, 2) = 0.5;
-	CameraMatrix.at<double>(1, 2) = 0.5;
-
-	FTransform OptimalCameraPose = NodalOffset * ExistingOffset.Inverse() * CameraPose;
-
-	// Convert the optimal camera pose to opencv's coordinate system
-	FOpenCVHelper::ConvertUnrealToOpenCV(OptimalCameraPose);
-
-	// Compute the reprojection error for this view
-	return FOpenCVHelper::ComputeReprojectionError(OptimalCameraPose, CameraMatrix, Points3d, Points2d);
-#else
-	return -1.0;
-#endif
+	return FCameraCalibrationSolver::OptimizeNodalOffset(ObjectPoints, ImagePoints, FocalLength, ImageCenter, CameraPoses, InOutNodalOffset);
 }
 
 void UCameraNodalOffsetAlgoPoints::Initialize(UNodalOffsetTool* InNodalOffsetTool)
@@ -1034,12 +786,15 @@ bool UCameraNodalOffsetAlgoPoints::CalculatedOptimalCameraComponentPose(
 		return false;
 	}
 
-	const FCameraCalibrationStepsController* StepsController;
-	const ULensFile* LensFile;
+	if (!NodalOffsetTool.IsValid())
+	{
+		return false;
+	}
 
-	if (!ensure(GetStepsControllerAndLensFile(&StepsController, &LensFile)))
+	const FCameraCalibrationStepsController* StepsController = NodalOffsetTool->GetCameraCalibrationStepsController();
+
+	if (!StepsController)
 	{
-		OutErrorMessage = LOCTEXT("ToolNotFound", "Tool not found");
 		return false;
 	}
 
@@ -1050,127 +805,33 @@ bool UCameraNodalOffsetAlgoPoints::CalculatedOptimalCameraComponentPose(
 		return false;
 	}
 
-	// Get parameters from the handler
-	FLensDistortionState DistortionState = DistortionHandler->GetCurrentDistortionState();
-
-#if WITH_OPENCV
-
-	// Find the pose that minimizes the reprojection error
+	TArray<FVector> ObjectPoints;
+	ObjectPoints.Reserve(Rows.Num());
 
-	// Populate the 3d/2d correlation points
-
-	std::vector<cv::Point3f> Points3d;
-	std::vector<cv::Point2f> Points2d;
-	Points3d.reserve(Rows.Num());
-	Points2d.reserve(Rows.Num());
+	TArray<FVector2f> ImagePoints;
+	ImagePoints.Reserve(Rows.Num());
 
 	for (const TSharedPtr<FNodalOffsetPointsRowData>& Row : Rows)
 	{
-		// Convert from UE coordinates to OpenCV coordinates
-		FTransform Transform;
-		Transform.SetIdentity();
-		Transform.SetLocation(Row->CalibratorPointData.Location);
-
-		FOpenCVHelper::ConvertUnrealToOpenCV(Transform);
-
-		// Calibrator 3d points
-		Points3d.push_back(cv::Point3f(
-			Transform.GetLocation().X,
-			Transform.GetLocation().Y,
-			Transform.GetLocation().Z));
-
-		Points2d.push_back(cv::Point2f(
-			Row->UndistortedPoint2D.X,
-			Row->UndistortedPoint2D.Y));
+		ObjectPoints.Add(Row->CalibratorPointData.Location);
+		ImagePoints.Add(Row->UndistortedPoint2D);
 	}
-	// Populate camera matrix
 
-	cv::Mat CameraMatrix(3, 3, cv::DataType<double>::type);
-	cv::setIdentity(CameraMatrix);
-
-	// Note: cv::Mat uses (row,col) indexing.
-	//
-	//  Fx  0  Cx
-	//  0  Fy  Cy
-	//  0   0   1
-
-	CameraMatrix.at<double>(0, 0) = DistortionState.FocalLengthInfo.FxFy.X;
-	CameraMatrix.at<double>(1, 1) = DistortionState.FocalLengthInfo.FxFy.Y;
-
-	// The displacement map will correct for image center offset
-	CameraMatrix.at<double>(0, 2) = 0.5;
-	CameraMatrix.at<double>(1, 2) = 0.5;
+	const FLensDistortionState DistortionState = DistortionHandler->GetCurrentDistortionState();
+	const FVector2D FocalLength = DistortionState.FocalLengthInfo.FxFy;
+	const FVector2D ImageCenter = FVector2D(0.5, 0.5);
 
-	// Solve for camera position
-	cv::Mat Rrod = cv::Mat::zeros(3, 1, cv::DataType<double>::type); // Rotation vector in Rodrigues notation. 3x1.
-	cv::Mat Tobj = cv::Mat::zeros(3, 1, cv::DataType<double>::type); // Translation vector. 3x1.
+	TArray<float> DistortionParameters = { 0.0f, 0.0f, 0.0f, 0.0f, 0.0f };
 
 	// We send no distortion parameters, because Points2d was manually undistorted already
-	if (!cv::solvePnP(Points3d, Points2d, CameraMatrix, cv::noArray(), Rrod, Tobj))
+	bool bResult = FOpenCVHelper::SolvePnP(ObjectPoints, ImagePoints, FocalLength, ImageCenter, DistortionParameters, OutDesiredCameraTransform);
+	if (!bResult)
 	{
 		OutErrorMessage = LOCTEXT("SolvePnpFailed", "Failed to resolve a camera position given the data in the calibration rows. Please retry the calibration.");
 		return false;
 	}
 
-	// Check for invalid data
-	{
-		const double Tx = Tobj.at<double>(0);
-		const double Ty = Tobj.at<double>(1);
-		const double Tz = Tobj.at<double>(2);
-
-		const double MaxValue = 1e16;
-
-		if (abs(Tx) > MaxValue || abs(Ty) > MaxValue || abs(Tz) > MaxValue)
-		{
-			OutErrorMessage = LOCTEXT("DataOutOfBounds", "The triangulated camera position had invalid values, please retry the calibration.");
-			return false;
-		}
-	}
-
-	// Convert to camera pose
-
-	// [R|t]' = [R'|-R'*t]
-
-	// Convert from Rodrigues to rotation matrix
-	cv::Mat Robj;
-	cv::Rodrigues(Rrod, Robj); // Robj is 3x3
-
-	// Calculate camera translation
-	cv::Mat Tcam = -Robj.t() * Tobj;
-
-	// Invert/transpose to get camera orientation
-	cv::Mat Rcam = Robj.t();
-
-	// Convert back to UE coordinates
-
-	FMatrix M = FMatrix::Identity;
-
-	// Fill rotation matrix
-	for (int32 Column = 0; Column < 3; ++Column)
-	{
-		M.SetColumn(Column, FVector(
-			Rcam.at<double>(Column, 0),
-			Rcam.at<double>(Column, 1),
-			Rcam.at<double>(Column, 2))
-		);
-	}
-
-	// Fill translation vector
-	M.M[3][0] = Tcam.at<double>(0);
-	M.M[3][1] = Tcam.at<double>(1);
-	M.M[3][2] = Tcam.at<double>(2);
-
-	OutDesiredCameraTransform.SetFromMatrix(M);
-	FOpenCVHelper::ConvertOpenCVToUnreal(OutDesiredCameraTransform);
-
 	return true;
-
-#else
-	{
-		OutErrorMessage = LOCTEXT("OpenCVRequired", "OpenCV is required");
-		return false;
-	}
-#endif //WITH_OPENCV
 }
 
 bool UCameraNodalOffsetAlgoPoints::CalibratorMovedInAnyRow(
@@ -2518,15 +2179,15 @@ void UCameraNodalOffsetAlgoPoints::UndistortCalibrationRowPoints()
 
 	const int32 NumPoints = CalibrationRows.Num();
 
-	TArray<FVector2D> ImagePoints;
+	TArray<FVector2f> ImagePoints;
 	ImagePoints.Reserve(NumPoints);
 
 	for (const TSharedPtr<FNodalOffsetPointsRowData>& Row : CalibrationRows)
 	{
-		ImagePoints.Add(FVector2D(Row->Point2D.X, Row->Point2D.Y));
+		ImagePoints.Add(Row->Point2D);
 	}
 
-	TArray<FVector2D> UndistortedPoints;
+	TArray<FVector2f> UndistortedPoints;
 	UndistortedPoints.AddZeroed(ImagePoints.Num());
 	DistortionRenderingUtils::UndistortImagePoints(DistortionHandler->GetDistortionDisplacementMap(), ImagePoints, UndistortedPoints);
 
diff --git a/Engine/Plugins/VirtualProduction/CameraCalibration/Source/CameraCalibrationEditor/Private/Calibrators/CameraNodalOffsetAlgoPoints.h b/Engine/Plugins/VirtualProduction/CameraCalibration/Source/CameraCalibrationEditor/Private/Calibrators/CameraNodalOffsetAlgoPoints.h
--- a/Engine/Plugins/VirtualProduction/CameraCalibration/Source/CameraCalibrationEditor/Private/Calibrators/CameraNodalOffsetAlgoPoints.h
+++ b/Engine/Plugins/VirtualProduction/CameraCalibration/Source/CameraCalibrationEditor/Private/Calibrators/CameraNodalOffsetAlgoPoints.h
@@ -119,11 +119,11 @@ struct FNodalOffsetPointsRowData
 
 	// Normalized 0~1 2d location of the identified calibrator point in the media plate.
 	UPROPERTY()
-	FVector2D Point2D = FVector2D::ZeroVector;
+	FVector2f Point2D = FVector2f::ZeroVector;
 
 	// Location of Point2D after it has been undistorted
 	UPROPERTY()
-	FVector2D UndistortedPoint2D = FVector2D::ZeroVector;
+	FVector2f UndistortedPoint2D = FVector2f::ZeroVector;
 
 	// True if this row has had its Point2D undistorted
 	UPROPERTY()
@@ -405,9 +405,6 @@ protected:
 	/* Refine the input nodal offset by running a minimization algorithm designed to minimize the sum of the reprojection errors for the camera views in the set of input pose groups*/
 	double MinimizeReprojectionError(FTransform& InOutNodalOffset, const TArray<TSharedPtr<TArray<TSharedPtr<FNodalOffsetPointsRowData>>>>& SamePoseRowGroups) const;
 
-	/* Compute the reprojection error for the rows in the input PoseGroup using the input NodalOffset */
-	double ComputeReprojectionError(const FTransform& NodalOffset, const TArray<TSharedPtr<FNodalOffsetPointsRowData>>& PoseGroup) const;
-
 	/** Export global session data to a .json file */
 	void ExportSessionData();
 
diff --git a/Engine/Plugins/VirtualProduction/CameraCalibration/Source/CameraCalibrationEditor/Private/CameraCalibrationUtilsPrivate.cpp b/Engine/Plugins/VirtualProduction/CameraCalibration/Source/CameraCalibrationEditor/Private/CameraCalibrationUtilsPrivate.cpp
new file mode 100644
--- /dev/null
+++ b/Engine/Plugins/VirtualProduction/CameraCalibration/Source/CameraCalibrationEditor/Private/CameraCalibrationUtilsPrivate.cpp
@@ -0,0 +1,127 @@
+// Copyright Epic Games, Inc. All Rights Reserved.
+
+#include "CameraCalibrationUtilsPrivate.h"
+
+#include "CalibrationPointComponent.h"
+#include "Engine/Texture2D.h"
+#include "TextureResource.h"
+
+namespace UE::CameraCalibration::Private
+{
+	inline static const TMap<EArucoDictionary, FString> ArucoDictionaries =
+	{
+		TPair<EArucoDictionary, FString>(EArucoDictionary::DICT_4X4_1000, TEXT("DICT_4X4_1000")),
+		TPair<EArucoDictionary, FString>(EArucoDictionary::DICT_4X4_250, TEXT("DICT_4X4_250")),
+		TPair<EArucoDictionary, FString>(EArucoDictionary::DICT_4X4_100, TEXT("DICT_4X4_100")),
+		TPair<EArucoDictionary, FString>(EArucoDictionary::DICT_4X4_50, TEXT("DICT_4X4_50")),
+		TPair<EArucoDictionary, FString>(EArucoDictionary::DICT_5X5_1000, TEXT("DICT_5X5_1000")),
+		TPair<EArucoDictionary, FString>(EArucoDictionary::DICT_5X5_250, TEXT("DICT_5X5_250")),
+		TPair<EArucoDictionary, FString>(EArucoDictionary::DICT_5X5_100, TEXT("DICT_5X5_100")),
+		TPair<EArucoDictionary, FString>(EArucoDictionary::DICT_5X5_50, TEXT("DICT_5X5_50")),
+		TPair<EArucoDictionary, FString>(EArucoDictionary::DICT_6X6_1000, TEXT("DICT_6X6_1000")),
+		TPair<EArucoDictionary, FString>(EArucoDictionary::DICT_6X6_250, TEXT("DICT_6X6_250")),
+		TPair<EArucoDictionary, FString>(EArucoDictionary::DICT_6X6_100, TEXT("DICT_6X6_100")),
+		TPair<EArucoDictionary, FString>(EArucoDictionary::DICT_6X6_50, TEXT("DICT_6X6_50")),
+		TPair<EArucoDictionary, FString>(EArucoDictionary::DICT_7X7_1000, TEXT("DICT_7X7_1000")),
+		TPair<EArucoDictionary, FString>(EArucoDictionary::DICT_7X7_250, TEXT("DICT_7X7_250")),
+		TPair<EArucoDictionary, FString>(EArucoDictionary::DICT_7X7_100, TEXT("DICT_7X7_100")),
+		TPair<EArucoDictionary, FString>(EArucoDictionary::DICT_7X7_50, TEXT("DICT_7X7_50")),
+		TPair<EArucoDictionary, FString>(EArucoDictionary::DICT_ARUCO_ORIGINAL, TEXT("DICT_ARUCO_ORIGINAL"))
+	};
+
+	EArucoDictionary GetArucoDictionaryFromName(FString Name)
+	{
+		TSet<EArucoDictionary> Keys;
+		ArucoDictionaries.GetKeys(Keys);
+
+		for (EArucoDictionary Key : Keys)
+		{
+			FString DictionaryName = GetArucoDictionaryName(Key);
+			if (Name.StartsWith(DictionaryName))
+			{
+				return Key;
+			}
+		}
+
+		return EArucoDictionary::None;
+	}
+
+	FString GetArucoDictionaryName(EArucoDictionary Dictionary)
+	{
+		const FString* DictionaryName = ArucoDictionaries.Find(Dictionary);
+
+		if (DictionaryName)
+		{
+			FString FoundDictionaryName = *DictionaryName;
+			return FoundDictionaryName;
+		}
+
+		return TEXT("");
+	}
+
+	EArucoDictionary GetArucoDictionaryForCalibrator(AActor* CalibratorActor)
+	{
+		if (!CalibratorActor)
+		{
+			return EArucoDictionary::None;
+		}
+
+		// Find all calibration components belonging to the input calibrator actor
+		constexpr uint32 NumInlineAllocations = 32;
+		TArray<UCalibrationPointComponent*, TInlineAllocator<NumInlineAllocations>> CalibrationComponents;
+		CalibratorActor->GetComponents(CalibrationComponents);
+
+		EArucoDictionary Dictionary = EArucoDictionary::None;
+		for (const UCalibrationPointComponent* Component : CalibrationComponents)
+		{
+			if (!Component)
+			{
+				continue;
+			}
+
+			// Look up the dictionary enum value based on the name of the component (only works if it is prefixed with a dictionary name)
+			Dictionary = GetArucoDictionaryFromName(Component->GetName());
+
+			if (Dictionary != EArucoDictionary::None)
+			{
+				break;
+			}
+
+			// Look up the dictionary enum value based on the name of each of the component's subpoints (only works if it is prefixed with a dictionary name)
+			for (const TPair<FString, FVector>& SubPoint : Component->SubPoints)
+			{
+				Dictionary = GetArucoDictionaryFromName(SubPoint.Key);
+
+				if (Dictionary != EArucoDictionary::None)
+				{
+					break;
+				}
+			}
+		}
+
+		return Dictionary;
+	}
+
+	void ClearTexture(UTexture2D* Texture, FColor ClearColor)
+	{
+		TArray<FColor> Pixels;
+		Pixels.Init(ClearColor, Texture->GetSizeX() * Texture->GetSizeY());
+
+		void* TextureData = Texture->GetPlatformData()->Mips[0].BulkData.Lock(LOCK_READ_WRITE);
+
+		FMemory::Memcpy(TextureData, Pixels.GetData(), Pixels.Num() * sizeof(FColor));
+
+		Texture->GetPlatformData()->Mips[0].BulkData.Unlock();
+		Texture->UpdateResource();
+	}
+
+	void SetTextureData(UTexture2D* Texture, const TArray<FColor>& PixelData)
+	{
+		void* TextureData = Texture->GetPlatformData()->Mips[0].BulkData.Lock(LOCK_READ_WRITE);
+
+		FMemory::Memcpy(TextureData, PixelData.GetData(), PixelData.Num() * sizeof(FColor));
+
+		Texture->GetPlatformData()->Mips[0].BulkData.Unlock();
+		Texture->UpdateResource();
+	}
+}
\ No newline at end of file
diff --git a/Engine/Plugins/VirtualProduction/CameraCalibration/Source/CameraCalibrationEditor/Private/CameraCalibrationUtilsPrivate.h b/Engine/Plugins/VirtualProduction/CameraCalibration/Source/CameraCalibrationEditor/Private/CameraCalibrationUtilsPrivate.h
new file mode 100644
--- /dev/null
+++ b/Engine/Plugins/VirtualProduction/CameraCalibration/Source/CameraCalibrationEditor/Private/CameraCalibrationUtilsPrivate.h
@@ -0,0 +1,22 @@
+// Copyright Epic Games, Inc. All Rights Reserved.
+#pragma once
+
+#include "OpenCVHelper.h"
+
+namespace UE::CameraCalibration::Private
+{
+	/** Returns the aruco dictionary that matches the calibration point component(s) of the input calibrator actor */
+	EArucoDictionary GetArucoDictionaryForCalibrator(AActor* CalibratorActor);
+
+	/** Get the aruco dictionary from a string representation matching the dictionary name */
+	EArucoDictionary GetArucoDictionaryFromName(FString Name);
+
+	/** Get the string representation of the input aruco dictionary */
+	FString GetArucoDictionaryName(EArucoDictionary Dictionary);
+
+	/** Set every pixel in the input texture to the clear color */
+	void ClearTexture(UTexture2D* Texture, FColor ClearColor = FColor::Transparent);
+
+	/** Set the texture data to the input array of pixels */
+	void SetTextureData(UTexture2D* Texture, const TArray<FColor>& PixelData);
+}
diff --git a/Engine/Plugins/VirtualProduction/CameraCalibrationCore/Source/CameraCalibrationCore/Private/DistortionRenderingUtils.cpp b/Engine/Plugins/VirtualProduction/CameraCalibrationCore/Source/CameraCalibrationCore/Private/DistortionRenderingUtils.cpp
--- a/Engine/Plugins/VirtualProduction/CameraCalibrationCore/Source/CameraCalibrationCore/Private/DistortionRenderingUtils.cpp
+++ b/Engine/Plugins/VirtualProduction/CameraCalibrationCore/Source/CameraCalibrationCore/Private/DistortionRenderingUtils.cpp
@@ -22,8 +22,8 @@ public:
 	BEGIN_SHADER_PARAMETER_STRUCT(FParameters, )
 		SHADER_PARAMETER_TEXTURE(Texture2D, DistortionMap)
 		SHADER_PARAMETER_SAMPLER(SamplerState, DistortionMapSampler)
-		SHADER_PARAMETER_SRV(StructuredBuffer<FVector2D>, InputPoints)
-		SHADER_PARAMETER_UAV(RWStructuredBuffer<FVector2D>, UndistortedPoints)
+		SHADER_PARAMETER_SRV(StructuredBuffer<FVector2f>, InputPoints)
+		SHADER_PARAMETER_UAV(RWStructuredBuffer<FVector2f>, UndistortedPoints)
 	END_SHADER_PARAMETER_STRUCT()
 
 public:
@@ -40,6 +40,27 @@ IMPLEMENT_GLOBAL_SHADER(FUndistortImagePointsCS, "/Plugin/CameraCalibrationCore/
 namespace DistortionRenderingUtils
 {
 	void UndistortImagePoints(UTextureRenderTarget2D* DistortionMap, TArray<FVector2D> ImagePoints, TArray<FVector2D>& OutUndistortedPoints)
+	{
+		// Convert the input points from doubles to floats
+		TArray<FVector2f> ImagePointsFloat;
+		ImagePointsFloat.Reserve(ImagePoints.Num());
+		for (FVector2D Point : ImagePoints)
+		{
+			ImagePointsFloat.Add(FVector2f(Point.X, Point.Y));
+		}
+
+		TArray<FVector2f> UndistortedPointsFloat;
+		UndistortImagePoints(DistortionMap, ImagePointsFloat, UndistortedPointsFloat);
+
+		// Convert the output points from floats to doubles
+		OutUndistortedPoints.Reserve(UndistortedPointsFloat.Num());
+		for (FVector2f Point : UndistortedPointsFloat)
+		{
+			OutUndistortedPoints.Add(FVector2D(Point.X, Point.Y));
+		}
+	}
+
+	void UndistortImagePoints(UTextureRenderTarget2D* DistortionMap, TArray<FVector2f> ImagePoints, TArray<FVector2f>& OutUndistortedPoints)
 	{
 		if ((DistortionMap == nullptr) || (ImagePoints.Num() < 1))
 		{
@@ -61,7 +82,7 @@ namespace DistortionRenderingUtils
 				InputPointsResourceArray.Reserve(NumPoints);
 				for (int32 Index = 0; Index < NumPoints; Index++)
 				{
-					InputPointsResourceArray.Add(FVector2f(Points[Index]));	// LWC_TODO: Precision loss
+					InputPointsResourceArray.Add(Points[Index]);
 				}
 
 				TResourceArray<FVector2f> EmptyBuffer;
@@ -100,7 +121,7 @@ namespace DistortionRenderingUtils
 				{
 					for (int32 Index = 0; Index < NumPoints; ++Index)
 					{
-						OutUndistortedPoints[Index] = FVector2D(UndistortedPointData[Index]);
+						OutUndistortedPoints[Index] = UndistortedPointData[Index];
 					}
 				}
 
diff --git a/Engine/Plugins/VirtualProduction/CameraCalibrationCore/Source/CameraCalibrationCore/Public/CameraCalibrationCheckerboard.h b/Engine/Plugins/VirtualProduction/CameraCalibrationCore/Source/CameraCalibrationCore/Public/CameraCalibrationCheckerboard.h
--- a/Engine/Plugins/VirtualProduction/CameraCalibrationCore/Source/CameraCalibrationCore/Public/CameraCalibrationCheckerboard.h
+++ b/Engine/Plugins/VirtualProduction/CameraCalibrationCore/Source/CameraCalibrationCore/Public/CameraCalibrationCheckerboard.h
@@ -64,15 +64,15 @@ public:
 	TObjectPtr<UCalibrationPointComponent> Center;
 
 
-	/** Number of rows. It is one of the parameters cv::findChessboardCorners needs */
+	/** Number of rows */
 	UPROPERTY(EditAnywhere, Category = "Calibration")
 	int32 NumCornerRows = 2;
 
-	/** Number of columns. It is one of the parameters cv::findChessboardCorners needs */
+	/** Number of columns */
 	UPROPERTY(EditAnywhere, Category = "Calibration")
 	int32 NumCornerCols = 2;
 
-	/** Length of the side of each square. It is used when building the object points that cv::calibrateCamera needs */
+	/** Length of the side of each square */
 	UPROPERTY(EditAnywhere, Category = "Calibration")
 	float SquareSideLength = 3.2089f;
 
diff --git a/Engine/Plugins/VirtualProduction/CameraCalibrationCore/Source/CameraCalibrationCore/Public/DistortionRenderingUtils.h b/Engine/Plugins/VirtualProduction/CameraCalibrationCore/Source/CameraCalibrationCore/Public/DistortionRenderingUtils.h
--- a/Engine/Plugins/VirtualProduction/CameraCalibrationCore/Source/CameraCalibrationCore/Public/DistortionRenderingUtils.h
+++ b/Engine/Plugins/VirtualProduction/CameraCalibrationCore/Source/CameraCalibrationCore/Public/DistortionRenderingUtils.h
@@ -9,6 +9,9 @@ class UTextureRenderTarget2D;
 
 namespace DistortionRenderingUtils
 {
-	/** Undistort the input array of 2D image points using the input displacement map */
+	UE_DEPRECATED(5.4, "The version of UndistortImagePoints that takes arrays of FVector2Ds as arguments is deprecated. Please use the version that take FVector2f arguments.")
 	CAMERACALIBRATIONCORE_API void UndistortImagePoints(UTextureRenderTarget2D* DistortionMap, TArray<FVector2D> ImagePoints, TArray<FVector2D>& OutUndistortedPoints);
+
+	/** Undistort the input array of 2D image points using the input displacement map */
+	CAMERACALIBRATIONCORE_API void UndistortImagePoints(UTextureRenderTarget2D* DistortionMap, TArray<FVector2f> ImagePoints, TArray<FVector2f>& OutUndistortedPoints);
 };
