OpenCV: Add check before calling solvePnP that enough points were provided to prevent a crash inside the opencv library.

diff --git a/Engine/Plugins/Runtime/OpenCV/Source/OpenCVHelper/Private/OpenCVHelper.cpp b/Engine/Plugins/Runtime/OpenCV/Source/OpenCVHelper/Private/OpenCVHelper.cpp
--- a/Engine/Plugins/Runtime/OpenCV/Source/OpenCVHelper/Private/OpenCVHelper.cpp
+++ b/Engine/Plugins/Runtime/OpenCV/Source/OpenCVHelper/Private/OpenCVHelper.cpp
@@ -565,6 +565,14 @@ bool FOpenCVHelper::SolvePnP(const TArray<FVector>& ObjectPoints, const TArray<F
 	{
 		return false;
 	}
+	
+	// For non-planar sets of 3D points, solvePnP requires a minimum of 6 points to compute the direct linear transformation (DLT)
+	constexpr int32 MinimumPoints = 6;
+	if (NumPoints < MinimumPoints)
+	{
+		UE_LOG(LogOpenCVHelper, Error, TEXT("SolvePnP requires a minimum of 6 3D/2D point correspondences, but only %d were provided"), NumPoints);
+		return false;
+	}
 
 	// cv::solvePnP() will only accept spherical distortion parameters, but it accepts a variable number of parameters based on how much of the distortion model is used
 	// We need to guard against an incorrect number of parameters to avoid crashing in the opencv module
