Add TPS and LICENSE file for OpenCV python package. Add python requirement to OpenCV.uplugin for the pip autoinstaller to install cv2

diff --git a/Engine/Plugins/Runtime/OpenCV/OpenCV.uplugin b/Engine/Plugins/Runtime/OpenCV/OpenCV.uplugin
--- a/Engine/Plugins/Runtime/OpenCV/OpenCV.uplugin
+++ b/Engine/Plugins/Runtime/OpenCV/OpenCV.uplugin
@@ -1,30 +1,37 @@
 {
-	"FileVersion" : 3,
-	"Version" : 1,
-	"VersionName" : "1.0",
-	"FriendlyName" : "OpenCV",
-	"Description" : "Plugin initializing OpenCV library to be used in engine.",
-	"Category" : "Computer Vision",
-	"CreatedBy" : "Epic Games, Inc.",
-	"CreatedByURL" : "https://epicgames.com",
-	"DocsURL" : "",
-	"MarketplaceURL" : "",
-	"SupportURL" : "",
-	"EnabledByDefault" : false,
-	"CanContainContent" : true,
-	"IsBetaVersion" : true,
-	"Installed" : false,
-	"Modules" :
-	[
+	"FileVersion": 3,
+	"Version": 1,
+	"VersionName": "1.0",
+	"FriendlyName": "OpenCV",
+	"Description": "Plugin initializing OpenCV library to be used in engine.",
+	"Category": "Computer Vision",
+	"CreatedBy": "Epic Games, Inc.",
+	"CreatedByURL": "https://epicgames.com",
+	"DocsURL": "",
+	"MarketplaceURL": "",
+	"SupportURL": "",
+	"EnabledByDefault": false,
+	"CanContainContent": true,
+	"IsBetaVersion": true,
+	"Installed": false,
+	"Modules": [
 		{
-			"Name" : "OpenCVHelper",
-			"Type" : "Runtime",
+			"Name": "OpenCVHelper",
+			"Type": "Runtime",
 			"LoadingPhase": "PostConfigInit",
 			"PlatformAllowList": [
 				"Win64",
 				"Linux",
 				"Mac"
-		  ]
+			]
+		}
+	],
+	"PythonRequirements": [
+		{
+			"Platform": "All",
+			"Requirements": [
+				"opencv-python==4.5.5.62"
+			]
 		}
 	]
 }
\ No newline at end of file
