PipInstall: Temporarily fix opencv plugin dependencies in case other python packages are not included. Should be permanently fixed using plugin dependency system

diff --git a/Engine/Plugins/Runtime/OpenCV/OpenCV.uplugin b/Engine/Plugins/Runtime/OpenCV/OpenCV.uplugin
--- a/Engine/Plugins/Runtime/OpenCV/OpenCV.uplugin
+++ b/Engine/Plugins/Runtime/OpenCV/OpenCV.uplugin
@@ -30,7 +30,8 @@
 		{
 			"Platform": "All",
 			"Requirements": [
-				"opencv-python==4.5.5.62 --hash=sha256:2601388def0d6b957cc30dd88f8ff74a5651ae6940dd9e488241608cfa2b15c7 --hash=sha256:130cc75d56b29aa3c5de8b6ac438242dd2574ba6eaa8bccdffdcfd6b78632f7f --hash=sha256:c463d2276d8662b972d20ca9644702188507de200ca5405b89e1fe71c5c99989 --hash=sha256:ac92e743e22681f30001942d78512c1e39bce53dbffc504e5645fdc45c0f2c47"
+				"opencv-python==4.5.5.62 --hash=sha256:2601388def0d6b957cc30dd88f8ff74a5651ae6940dd9e488241608cfa2b15c7 --hash=sha256:130cc75d56b29aa3c5de8b6ac438242dd2574ba6eaa8bccdffdcfd6b78632f7f --hash=sha256:c463d2276d8662b972d20ca9644702188507de200ca5405b89e1fe71c5c99989 --hash=sha256:ac92e743e22681f30001942d78512c1e39bce53dbffc504e5645fdc45c0f2c47",
+				"numpy==1.20.3 --hash=sha256:830b044f4e64a76ba71448fce6e604c0fc47a0e54d8f6467be23749ac2cbd2fb --hash=sha256:5d050e1e4bc9ddb8656d7b4f414557720ddcca23a5b88dd7cff65e847864c400 --hash=sha256:6690080810f77485667bfbff4f69d717c3be25e5b11bb2073e76bb3f578d99b4"
 			]
 		}
 	]
