OpenCVHelper: Add variants to ProjectPoints() and DrawChessboardCorners() that operate on FVector2D points (in addition to the existing methods that operate on FVector2f).
This is primarily useful in camera calibration to avoid unneccessary copies from double to float.
However, the float variants may still be useful in other contexts, so there is no need to deprecate the other versions at this time.
#rnx

diff --git a/Engine/Plugins/Runtime/OpenCV/Source/OpenCVHelper/Private/OpenCVHelper.cpp b/Engine/Plugins/Runtime/OpenCV/Source/OpenCVHelper/Private/OpenCVHelper.cpp
--- a/Engine/Plugins/Runtime/OpenCV/Source/OpenCVHelper/Private/OpenCVHelper.cpp
+++ b/Engine/Plugins/Runtime/OpenCV/Source/OpenCVHelper/Private/OpenCVHelper.cpp
@@ -557,6 +557,18 @@ bool FOpenCVHelper::DrawCheckerboardCorners(const TArray<FVector2f>& Corners, FI
 #endif // WITH_OPENCV
 }
 
+bool FOpenCVHelper::DrawCheckerboardCorners(const TArray<FVector2D>& Corners, FIntPoint CheckerboardDimensions, UTexture2D* DebugTexture)
+{
+	TArray<FVector2f> CornersFloat;
+	CornersFloat.Reserve(Corners.Num());
+	for (const FVector2D& Corner : Corners)
+	{
+		CornersFloat.Add(FVector2f(Corner.X, Corner.Y));
+	}
+
+	return DrawCheckerboardCorners(CornersFloat, CheckerboardDimensions, DebugTexture);
+}
+
 bool FOpenCVHelper::SolvePnP(const TArray<FVector>& ObjectPoints, const TArray<FVector2f>& ImagePoints, const FVector2D& FocalLength, const FVector2D& ImageCenter, const TArray<float>& DistortionParameters, FTransform& OutCameraPose)
 {
 #if WITH_OPENCV
@@ -626,6 +638,21 @@ bool FOpenCVHelper::SolvePnP(const TArray<FVector>& ObjectPoints, const TArray<F
 }
 
 bool FOpenCVHelper::ProjectPoints(const TArray<FVector>& ObjectPoints, const FVector2D& FocalLength, const FVector2D& ImageCenter, const TArray<float>& DistortionParameters, const FTransform& CameraPose, TArray<FVector2f>& OutImagePoints)
+{
+	TArray<FVector2D> OutImagePointsDoublePrecision;
+	if (ProjectPoints(ObjectPoints, FocalLength, ImageCenter, DistortionParameters, CameraPose, OutImagePointsDoublePrecision))
+	{
+		for (const FVector2D& Point : OutImagePointsDoublePrecision)
+		{
+			OutImagePoints.Add(FVector2f(Point.X, Point.Y));
+		}
+
+		return true;
+	}
+	return false;
+}
+
+bool FOpenCVHelper::ProjectPoints(const TArray<FVector>& ObjectPoints, const FVector2D& FocalLength, const FVector2D& ImageCenter, const TArray<float>& DistortionParameters, const FTransform& CameraPose, TArray<FVector2D>& OutImagePoints)
 {
 #if WITH_OPENCV
 	const int32 NumPoints = ObjectPoints.Num();
@@ -655,18 +682,11 @@ bool FOpenCVHelper::ProjectPoints(const TArray<FVector>& ObjectPoints, const FVe
 
 	cv::Mat DistortionParametersMat = cv::Mat(DistortionParameters.Num(), 1, CV_32FC1, (void*)DistortionParameters.GetData());
 
-	// cv::projectPoints requires that the 3D points and 2D points have the same bit depth, so we compute projected points with double precision, and then convert them back to floats before outputting
-	TArray<FVector2D> ImagePointsDoublePrecision;
-	ImagePointsDoublePrecision.Init(FVector2D(0.0, 0.0), NumPoints);
-	cv::Mat ImagePointsMat = cv::Mat(NumPoints, 1, CV_64FC2, (void*)ImagePointsDoublePrecision.GetData());
+	OutImagePoints.Init(FVector2D(0.0, 0.0), NumPoints);
+	cv::Mat ImagePointsMat = cv::Mat(NumPoints, 1, CV_64FC2, (void*)OutImagePoints.GetData());
 
 	cv::projectPoints(ObjectPointsMat, Rotation, Translation, CameraMatrix, DistortionParametersMat, ImagePointsMat);
 
-	for (const FVector2D& Point : ImagePointsDoublePrecision)
-	{
-		OutImagePoints.Add(FVector2f(Point.X, Point.Y));
-	}
-
 	return true;
 #else
 	return false;
diff --git a/Engine/Plugins/Runtime/OpenCV/Source/OpenCVHelper/Public/OpenCVHelper.h b/Engine/Plugins/Runtime/OpenCV/Source/OpenCVHelper/Public/OpenCVHelper.h
--- a/Engine/Plugins/Runtime/OpenCV/Source/OpenCVHelper/Public/OpenCVHelper.h
+++ b/Engine/Plugins/Runtime/OpenCV/Source/OpenCVHelper/Public/OpenCVHelper.h
@@ -155,12 +155,18 @@ public:
 	/** Draw a debug view of the input checkerboard corners on top of the input texture */
 	static bool DrawCheckerboardCorners(const TArray<FVector2f>& Corners, FIntPoint CheckerboardDimensions, UTexture2D* DebugTexture);
 
+	/** Draw a debug view of the input checkerboard corners on top of the input texture */
+	static bool DrawCheckerboardCorners(const TArray<FVector2D>& Corners, FIntPoint CheckerboardDimensions, UTexture2D* DebugTexture);
+
 	/** Compute the camera pose that minimizes the reprojection error of the input object points and image points */
 	static bool SolvePnP(const TArray<FVector>& ObjectPoints, const TArray<FVector2f>& ImagePoints, const FVector2D& FocalLength, const FVector2D& ImageCenter, const TArray<float>& DistortionParameters, FTransform& OutCameraPose);
 
 	/** Project the input object points to the 2D image plane defined by the input camera intrinsics and camera pose */
 	static bool ProjectPoints(const TArray<FVector>& ObjectPoints, const FVector2D& FocalLength, const FVector2D& ImageCenter, const TArray<float>& DistortionParameters, const FTransform& CameraPose, TArray<FVector2f>& OutImagePoints);
 
+	/** Project the input object points to the 2D image plane defined by the input camera intrinsics and camera pose */
+	static bool ProjectPoints(const TArray<FVector>& ObjectPoints, const FVector2D& FocalLength, const FVector2D& ImageCenter, const TArray<float>& DistortionParameters, const FTransform& CameraPose, TArray<FVector2D>& OutImagePoints);
+
 	/** Find a 3D fit line that passes through the input points, as well as a point on that line */
 	static bool FitLine3D(const TArray<FVector>& InPoints, FVector& OutLine, FVector& OutPointOnLine);
 
