CameraCalibration: Fix crash when detecting a checkerboard with marquee select that goes beyond the bounds of the media image.

diff --git a/Engine/Plugins/Runtime/OpenCV/Source/OpenCVHelper/Private/OpenCVHelper.cpp b/Engine/Plugins/Runtime/OpenCV/Source/OpenCVHelper/Private/OpenCVHelper.cpp
--- a/Engine/Plugins/Runtime/OpenCV/Source/OpenCVHelper/Private/OpenCVHelper.cpp
+++ b/Engine/Plugins/Runtime/OpenCV/Source/OpenCVHelper/Private/OpenCVHelper.cpp
@@ -493,12 +493,12 @@ bool FOpenCVHelper::IdentifyCheckerboard(TArray<FColor>& Image, FIntPoint ImageS
 	const cv::Mat WholeImageMat = cv::Mat(ImageSize.Y, ImageSize.X, CV_8UC4, Image.GetData());
 
 	// Sanitize the ROI to ensure that it lies completely within the bounds of the full size image
-	RegionOfInterest.Min.X = FMath::Clamp(RegionOfInterest.Min.X, 0, RegionOfInterest.Min.X);
-	RegionOfInterest.Min.Y = FMath::Clamp(RegionOfInterest.Min.Y, 0, RegionOfInterest.Min.Y);
-	RegionOfInterest.Max.X = FMath::Clamp(RegionOfInterest.Min.X, RegionOfInterest.Max.X, ImageSize.X);
-	RegionOfInterest.Max.Y = FMath::Clamp(RegionOfInterest.Min.Y, RegionOfInterest.Max.Y, ImageSize.Y);
+	RegionOfInterest.Min.X = FMath::Clamp(RegionOfInterest.Min.X, 0, ImageSize.X);
+	RegionOfInterest.Min.Y = FMath::Clamp(RegionOfInterest.Min.Y, 0, ImageSize.Y);
+	RegionOfInterest.Max.X = FMath::Clamp(RegionOfInterest.Max.X, 0, ImageSize.X);
+	RegionOfInterest.Max.Y = FMath::Clamp(RegionOfInterest.Max.Y, 0, ImageSize.Y);
 
-	if (RegionOfInterest.IsEmpty())
+	if ((RegionOfInterest.Width() <= 0) || (RegionOfInterest.Height() <= 0))
 	{
 		OutCorners.Empty();
 		return false;
