Opencv: Update uplugin numpy pin version to match python foundation packages.

diff --git a/Engine/Plugins/Runtime/OpenCV/OpenCV.uplugin b/Engine/Plugins/Runtime/OpenCV/OpenCV.uplugin
--- a/Engine/Plugins/Runtime/OpenCV/OpenCV.uplugin
+++ b/Engine/Plugins/Runtime/OpenCV/OpenCV.uplugin
@@ -31,7 +31,7 @@
 			"Platform": "All",
 			"Requirements": [
 				"opencv-python==4.5.5.62 --hash=sha256:c463d2276d8662b972d20ca9644702188507de200ca5405b89e1fe71c5c99989 --hash=sha256:130cc75d56b29aa3c5de8b6ac438242dd2574ba6eaa8bccdffdcfd6b78632f7f --hash=sha256:ac92e743e22681f30001942d78512c1e39bce53dbffc504e5645fdc45c0f2c47 --hash=sha256:2601388def0d6b957cc30dd88f8ff74a5651ae6940dd9e488241608cfa2b15c7",
-				"numpy==1.24.4 --hash=sha256:2541312fbf09977f3b3ad449c4e5f4bb55d0dbf79226d7724211acc905049400 --hash=sha256:d11efb4dbecbdf22508d55e48d9c8384db795e1b7b51ea735289ff96613ff74d --hash=sha256:7215847ce88a85ce39baf9e89070cb860c98fdddacbaa6c0da3ffb31b3350bd5 --hash=sha256:9667575fb6d13c95f1b36aca12c5ee3356bf001b714fc354eb5465ce1609e62f --hash=sha256:ed094d4f0c177b1b8e7aa9cba7d6ceed51c0e569a5318ac0ca9a090680a6a1b1 --hash=sha256:b4bea75e47d9586d31e892a7401f76e909712a0fd510f58f5337bea9572c571e --hash=sha256:7ffe43c74893dbf38c2b0a1f5428760a1a9c98285553c89e12d70a96a7f3a4d6 --hash=sha256:c0bfb52d2169d58c1cdb8cc1f16989101639b34c7d3ce60ed70b19c63eba0b64 --hash=sha256:befe2bf740fd8373cf56149a5c23a0f601e82869598d41f8e188a0e9869926f8 --hash=sha256:b7b1fc9864d7d39e28f41d089bfd6353cb5f27ecd9905348c24187a768c79694 --hash=sha256:e2926dac25b313635e4d6cf4dc4e51c8c0ebfed60b801c799ffc4c32bf3d1254 --hash=sha256:f136bab9c2cfd8da131132c2cf6cc27331dd6fae65f95f69dcd4ae3c3639c810"
+				"numpy==1.26.4 --hash=sha256:03a8c78d01d9781b28a6989f6fa1bb2c4f2d51201cf99d3dd875df6fbd96b23b --hash=sha256:08beddf13648eb95f8d867350f6a018a4be2e5ad54c8d8caed89ebca558b2818 --hash=sha256:2e4ee3380d6de9c9ec04745830fd9e2eccb3e6cf790d39d7b98ffd19b0dd754a --hash=sha256:3373d5d70a5fe74a2c1bb6d2cfd9609ecf686d47a2d7b1d37a8f3b6bf6003aea --hash=sha256:4c66707fabe114439db9068ee468c26bbdf909cac0fb58686a42a24de1760c71 --hash=sha256:52b8b60467cd7dd1e9ed082188b4e6bb35aa5cdd01777621a1658910745b90be --hash=sha256:666dbfb6ec68962c033a450943ded891bed2d54e6755e35e5835d63f4f6931d5 --hash=sha256:675d61ffbfa78604709862923189bad94014bef562cc35cf61d3a07bba02a7ed --hash=sha256:7349ab0fa0c429c82442a27a9673fc802ffdb7c7775fad780226cb234965e53c --hash=sha256:9ff0f4f29c51e2803569d7a51c2304de5554655a60c5d776e35b4a41413830d0 --hash=sha256:b3ce300f3644fb06443ee2222c2201dd3a89ea6040541412b8fa189341847218 --hash=sha256:b97fe8060236edf3662adfc2c633f56a08ae30560c56310562cb4f95500022d5 --hash=sha256:cd25bcecc4974d09257ffcd1f098ee778f7834c3ad767fe5db785be9a4aa9cb2 --hash=sha256:edd8b5fe47dab091176d21bb6de568acdd906d1887a4584a15a9a96a1dca06ef --hash=sha256:f870204a840a60da0b12273ef34f7051e98c3b5961b61b0c2c1be6dfd64fbcd3 --hash=sha256:ffa75af20b44f8dba823498024771d5ac50620e6915abac414251bd971b4529f"
 			]
 		}
 	]
