PythonFoundationPackages: Fix for incorrect dependency specification in OpenCV.

diff --git a/Engine/Plugins/Runtime/OpenCV/OpenCV.uplugin b/Engine/Plugins/Runtime/OpenCV/OpenCV.uplugin
--- a/Engine/Plugins/Runtime/OpenCV/OpenCV.uplugin
+++ b/Engine/Plugins/Runtime/OpenCV/OpenCV.uplugin
@@ -18,7 +18,7 @@
 		{
 			"Name": "PythonFoundationPackages",
 			"Enabled": true,
-			"PlatformAllowList": [
+			"TargetAllowList": [
 				"Editor"
 			]
 		}
