CameraCalibration: Add CharUco pattern to lens calibration and nodal offset.

Includes a charuco board calibration pattern that is configurable (sizing and aruco dict to use).

Also fixed a crash and handling in readmedia pixels related to how the camera is fit into the destination texture.

Added the ctrl+alt+drag not in the simulcam viewport tooltip for better discoverability of that fature.

#virtualized

diff --git a/Engine/Plugins/Runtime/OpenCV/Source/OpenCVHelper/Private/OpenCVHelper.cpp b/Engine/Plugins/Runtime/OpenCV/Source/OpenCVHelper/Private/OpenCVHelper.cpp
--- a/Engine/Plugins/Runtime/OpenCV/Source/OpenCVHelper/Private/OpenCVHelper.cpp
+++ b/Engine/Plugins/Runtime/OpenCV/Source/OpenCVHelper/Private/OpenCVHelper.cpp
@@ -3,6 +3,7 @@
 #include "OpenCVHelper.h"
 
 #include "Engine/Texture2D.h"
+#include "Misc/ScopeExit.h"
 #include "TextureResource.h"
 #include "UObject/Package.h"
 
@@ -11,6 +12,7 @@
 #include "PreOpenCVHeaders.h" // IWYU pragma: keep
 #include "opencv2/aruco.hpp"
 #include "opencv2/calib3d.hpp"
+#include "opencv2/aruco/charuco.hpp"
 #include "opencv2/imgproc.hpp"
 #include "PostOpenCVHeaders.h" // IWYU pragma: keep
 
@@ -363,6 +365,15 @@ void FOpenCVHelper::MakeObjectVectorsFromCameraPose(const FTransform& InTransfor
 bool FOpenCVHelper::IdentifyArucoMarkers(TArray<FColor>& Image, FIntPoint ImageSize, EArucoDictionary DictionaryName, TArray<FArucoMarker>& OutMarkers)
 {
 #if WITH_OPENCV
+	// Validate buffer size matches dimensions
+	const int32 RequiredPixels = ImageSize.X * ImageSize.Y;
+	if (Image.Num() < RequiredPixels)
+	{
+		UE_LOG(LogOpenCVHelper, Error, TEXT("IdentifyArucoMarkers: Image buffer too small. Expected %d pixels, got %d"), 
+			RequiredPixels, Image.Num());
+		return false;
+	}
+
 	// Initialize an OpenCV matrix header to point at the input image data. 
 	// The format for FColor is B8G8R8A8, so the corresponding OpenCV format is CV_8UC4 (8-bit per channel, 4 channels)
 	cv::Mat ImageMat = cv::Mat(ImageSize.Y, ImageSize.X, CV_8UC4, Image.GetData());
@@ -490,6 +501,15 @@ bool FOpenCVHelper::IdentifyCheckerboard(TArray<FColor>& Image, FIntPoint ImageS
 bool FOpenCVHelper::IdentifyCheckerboard(TArray<FColor>& Image, FIntPoint ImageSize, FIntRect RegionOfInterest, FIntPoint CheckerboardDimensions, TArray<FVector2f>& OutCorners)
 {
 #if WITH_OPENCV
+	// Validate buffer size matches dimensions
+	const int32 RequiredPixels = ImageSize.X * ImageSize.Y;
+	if (Image.Num() < RequiredPixels)
+	{
+		UE_LOG(LogOpenCVHelper, Error, TEXT("IdentifyCheckerboard: Image buffer too small. Expected %d pixels, got %d"), 
+			RequiredPixels, Image.Num());
+		return false;
+	}
+
 	// Initialize an OpenCV matrix header to point at the input image data. 
 	// The format for FColor is B8G8R8A8, so the corresponding OpenCV format is CV_8UC4 (8-bit per channel, 4 channels)
 	const cv::Mat WholeImageMat = cv::Mat(ImageSize.Y, ImageSize.X, CV_8UC4, Image.GetData());
@@ -904,3 +924,410 @@ cv::Mat FOpenCVLensDistortionParametersBase::CreateOpenCVCameraMatrix(const FVec
 	return CameraMatrix;
 }
 #endif // WITH_OPENCV
+
+
+bool FOpenCVHelper::IdentifyCharucoCorners(TArray<FColor>& Image, FIntPoint ImageSize, const FCharucoBoardConfig& BoardConfig, FCharucoCorners& OutCorners, const FIntRect& ROI)
+{
+#if WITH_OPENCV
+	if (Image.IsEmpty())
+	{
+		UE_LOG(LogOpenCVHelper, Error, TEXT("IdentifyCharucoCorners: Image is empty"));
+		return false;
+	}
+
+	// Validate buffer size matches dimensions
+	const int32 RequiredPixels = ImageSize.X * ImageSize.Y;
+	if (Image.Num() < RequiredPixels)
+	{
+		UE_LOG(LogOpenCVHelper, Error, TEXT("IdentifyCharucoCorners: Image buffer too small. Expected %d pixels, got %d"), 
+			RequiredPixels, Image.Num());
+		return false;
+	}
+
+	// Get ArUco dictionary
+
+	cv::Ptr<cv::aruco::Dictionary> Dictionary = UE::Aruco::Private::GetArucoDictionary(BoardConfig.Dictionary);
+
+	if (Dictionary == nullptr)
+	{
+		UE_LOG(LogOpenCVHelper, Error, TEXT("IdentifyCharucoCorners: Invalid ArUco dictionary"));
+		return false;
+	}
+
+	// Create Charuco board
+
+	cv::Ptr<cv::aruco::CharucoBoard> CharucoBoard = cv::aruco::CharucoBoard::create(
+		BoardConfig.SquaresX, 
+		BoardConfig.SquaresY, 
+		BoardConfig.SquareSize, 
+		BoardConfig.MarkerSize, 
+		Dictionary
+	);
+
+	// Convert image to OpenCV format
+
+	cv::Mat CvImage(ImageSize.Y, ImageSize.X, CV_8UC4, Image.GetData());
+	cv::Mat GrayImage;
+	cv::cvtColor(CvImage, GrayImage, cv::COLOR_BGRA2GRAY);
+
+	// Apply ROI (if any)
+
+	bool bUsingROI = (ROI.Width() > 0 && ROI.Height() > 0);
+	cv::Mat ProcessingImage = GrayImage;
+	FIntPoint ROIOffset = FIntPoint::ZeroValue;
+	
+	if (bUsingROI)
+	{
+		// Clamp ROI to image bounds
+
+		int ROILeft = FMath::Clamp(ROI.Min.X, 0, ImageSize.X - 1);
+		int ROITop = FMath::Clamp(ROI.Min.Y, 0, ImageSize.Y - 1);
+		int ROIRight = FMath::Clamp(ROI.Max.X, ROI.Min.X + 1, ImageSize.X);
+		int ROIBottom = FMath::Clamp(ROI.Max.Y, ROI.Min.Y + 1, ImageSize.Y);
+		
+		ROIOffset = FIntPoint(ROILeft, ROITop);
+		cv::Rect ROIRect(ROILeft, ROITop, ROIRight - ROILeft, ROIBottom - ROITop);
+
+		ProcessingImage = GrayImage(ROIRect);
+		
+		UE_LOG(LogOpenCVHelper, Verbose, TEXT("IdentifyCharucoCorners: Using ROI (%d,%d) to (%d,%d), size %dx%d"), 
+			ROILeft, ROITop, ROIRight, ROIBottom, ProcessingImage.cols, ProcessingImage.rows);
+	}
+
+	// First detect ArUco markers
+
+	std::vector<int> MarkerIds;
+	std::vector<std::vector<cv::Point2f>> MarkerCorners;
+	cv::Ptr<cv::aruco::DetectorParameters> Parameters = cv::aruco::DetectorParameters::create();
+
+	cv::aruco::detectMarkers(ProcessingImage, Dictionary, MarkerCorners, MarkerIds, Parameters);
+
+	// If no markers found, return false
+
+	if (MarkerIds.empty())
+	{
+		UE_LOG(LogOpenCVHelper, Warning, TEXT("IdentifyCharucoCorners: No ArUco markers detected"));
+		return false;
+	}
+
+	// Store detected ArUco markers (transform coordinates back to full image space if using ROI)
+
+	OutCorners.Markers.Empty();
+
+	for (size_t MarkerIdx = 0; MarkerIdx < MarkerIds.size(); MarkerIdx++)
+	{
+		FArucoMarker Marker;
+		Marker.MarkerID = MarkerIds[MarkerIdx];
+
+		constexpr int32 CornersPerMarker = 4;
+
+		for (int32 CornerIdx = 0; CornerIdx < CornersPerMarker; CornerIdx++)
+		{
+			const float X = MarkerCorners[MarkerIdx][CornerIdx].x + (bUsingROI ? ROIOffset.X : 0);
+			const float Y = MarkerCorners[MarkerIdx][CornerIdx].y + (bUsingROI ? ROIOffset.Y : 0);
+
+			Marker.Corners[CornerIdx] = FVector2f(X, Y);
+		}
+
+		OutCorners.Markers.Add(Marker);
+	}
+
+	// Interpolate Charuco corners
+	std::vector<cv::Point2f> CharucoCorners;
+	std::vector<int> CharucoIds;
+	
+	if (MarkerIds.size() > 0)
+	{
+		cv::aruco::interpolateCornersCharuco(MarkerCorners, MarkerIds, ProcessingImage, CharucoBoard, CharucoCorners, CharucoIds);
+	}
+
+	// If no Charuco corners found, return false
+	if (CharucoIds.empty())
+	{
+		UE_LOG(LogOpenCVHelper, Warning, TEXT("IdentifyCharucoCorners: No Charuco corners detected"));
+		return false;
+	}
+
+	// Convert to output format (transform coordinates back to full image space if using ROI)
+	OutCorners.Corners.Empty();
+	OutCorners.CornerIds.Empty();
+	
+	for (size_t CornerIdx = 0; CornerIdx < CharucoCorners.size(); CornerIdx++)
+	{
+		const float X = CharucoCorners[CornerIdx].x + (bUsingROI ? ROIOffset.X : 0);
+		const float Y = CharucoCorners[CornerIdx].y + (bUsingROI ? ROIOffset.Y : 0);
+
+		OutCorners.Corners.Add(FVector2f(X, Y));
+		OutCorners.CornerIds.Add(CharucoIds[CornerIdx]);
+	}
+
+	UE_LOG(LogOpenCVHelper, Verbose, TEXT("IdentifyCharucoCorners: Detected %d Charuco corners and %d ArUco markers"), 
+		OutCorners.Corners.Num(), OutCorners.Markers.Num());
+
+	return true;
+#else
+	UE_LOG(LogOpenCVHelper, Error, TEXT("IdentifyCharucoCorners: OpenCV not available"));
+	return false;
+#endif
+}
+
+
+bool FOpenCVHelper::DrawCharucoCorners(const FCharucoCorners& Corners, UTexture2D* DebugTexture)
+{
+#if WITH_OPENCV
+
+	if (!DebugTexture)
+	{
+		UE_LOG(LogOpenCVHelper, Error, TEXT("DrawCharucoCorners: Invalid texture (null)"));
+		return false;
+	}
+
+	FTexturePlatformData* PlatData = DebugTexture->GetPlatformData();
+	if (!PlatData || PlatData->Mips.Num() == 0)
+	{
+		UE_LOG(LogOpenCVHelper, Error, TEXT("DrawCharucoCorners: Missing platform data or mips"));
+		return false;
+	}
+
+	if (PlatData->PixelFormat != PF_B8G8R8A8)
+	{
+		UE_LOG(LogOpenCVHelper, Error, TEXT("DrawCharucoCorners: Unsupported PixelFormat. Expected PF_B8G8R8A8, got %d"), (int32)PlatData->PixelFormat);
+		return false;
+	}
+
+	const int32 Width = DebugTexture->GetSizeX();
+	const int32 Height = DebugTexture->GetSizeY();
+
+	if (Width <= 0 || Height <= 0)
+	{
+		UE_LOG(LogOpenCVHelper, Error, TEXT("DrawCharucoCorners: Invalid texture size %dx%d"), Width, Height);
+		return false;
+	}
+
+	FTexture2DMipMap& Mip = PlatData->Mips[0];
+	uint8* TextureData = static_cast<uint8*>(Mip.BulkData.Lock(LOCK_READ_WRITE));
+
+	if (!TextureData)
+	{
+		UE_LOG(LogOpenCVHelper, Error, TEXT("DrawCharucoCorners: Failed to lock texture bulk data"));
+		return false;
+	}
+
+	// Ensure all exit paths from now on call Mip.BulkData.Unlock().
+	ON_SCOPE_EXIT
+	{
+		Mip.BulkData.Unlock();
+	};
+
+	const int32 BytesPerPixel = 4;
+
+	const SIZE_T ExpectedSize = SIZE_T(Width) * SIZE_T(Height) * SIZE_T(BytesPerPixel);
+	const SIZE_T BulkSize = Mip.BulkData.GetBulkDataSize();
+
+	if (BulkSize < ExpectedSize)
+	{
+		UE_LOG(LogOpenCVHelper, Error, TEXT("DrawCharucoCorners: BulkData too small. Expected >= %llu, got %llu"),
+			(unsigned long long)ExpectedSize, (unsigned long long)BulkSize);
+		return false;
+	}
+
+	const size_t Step = size_t(Width) * size_t(BytesPerPixel);
+	cv::Mat CvImageBGRA(Height, Width, CV_8UC4, TextureData, Step);
+
+	// Convert to 3-channel BGR for OpenCV drawing, because it expects 1 or 3 channels, not e.g. 4.
+
+	cv::Mat CvImageBGR;
+	cv::cvtColor(CvImageBGRA, CvImageBGR, cv::COLOR_BGRA2BGR);
+
+	const bool bHasCharuco = !Corners.Corners.IsEmpty() && !Corners.CornerIds.IsEmpty();
+	const bool bHasMarkers = !Corners.Markers.IsEmpty();
+
+	if (!bHasCharuco && !bHasMarkers)
+	{
+		DebugTexture->UpdateResource();
+		return true;
+	}
+
+	if (bHasCharuco && (Corners.Corners.Num() != Corners.CornerIds.Num()))
+	{
+		UE_LOG(LogOpenCVHelper, Error, TEXT("DrawCharucoCorners: corners/ids size mismatch (%d vs %d)"),
+			Corners.Corners.Num(), Corners.CornerIds.Num());
+		return false;
+	}
+
+	std::vector<cv::Point2f> CharucoCorners;
+	std::vector<int> CharucoIds;
+
+	if (bHasCharuco)
+	{
+		CharucoCorners.reserve(Corners.Corners.Num());
+		for (const FVector2f& Corner : Corners.Corners)
+		{
+			CharucoCorners.emplace_back(Corner.X, Corner.Y);
+		}
+
+		CharucoIds.reserve(Corners.CornerIds.Num());
+		for (int32 CornerId : Corners.CornerIds)
+		{
+			CharucoIds.push_back((int)CornerId);
+		}
+	}
+
+	std::vector<std::vector<cv::Point2f>> MarkerCornerQuads;
+	std::vector<int> MarkerIds;
+
+	if (bHasMarkers)
+	{
+		MarkerCornerQuads.reserve(Corners.Markers.Num());
+		MarkerIds.reserve(Corners.Markers.Num());
+
+		for (const FArucoMarker& ArucoMarker : Corners.Markers)
+		{
+			std::vector<cv::Point2f> CornerQuad;
+
+			CornerQuad.reserve(4);
+
+			constexpr int32 CornersPerAruco = 4;
+
+			for (int32 CornerIdx = 0; CornerIdx < CornersPerAruco; ++CornerIdx)
+			{
+				const FVector2f& Corner = ArucoMarker.Corners[CornerIdx];
+				CornerQuad.emplace_back(Corner.X, Corner.Y);
+			}
+
+			MarkerCornerQuads.push_back(std::move(CornerQuad));
+			MarkerIds.push_back((int)ArucoMarker.MarkerID);
+		}
+	}
+
+	if (bHasCharuco)
+	{
+		// Draw corners as simple circles without IDs
+		for (size_t CornerIdx = 0; CornerIdx < CharucoCorners.size(); ++CornerIdx)
+		{
+			cv::circle(CvImageBGR, CharucoCorners[CornerIdx], 5, cv::Scalar(0, 255, 0), -1);
+		}
+	}
+	if (bHasMarkers)
+	{
+		// Draw marker corners without IDs
+		for (const std::vector<cv::Point2f>& MarkerCornerQuad : MarkerCornerQuads)
+		{
+			for (size_t CornerQuadIdx = 0; CornerQuadIdx < MarkerCornerQuad.size(); ++CornerQuadIdx)
+			{
+				cv::line(CvImageBGR, MarkerCornerQuad[CornerQuadIdx], MarkerCornerQuad[(CornerQuadIdx + 1) % MarkerCornerQuad.size()], cv::Scalar(255, 0, 255), 2);
+			}
+		}
+	}
+
+	// Convert back to BGRA for Unreal texture (preserve alpha channel by expanding BGR)
+	cv::cvtColor(CvImageBGR, CvImageBGRA, cv::COLOR_BGR2BGRA);
+
+	DebugTexture->UpdateResource();
+	return true;
+#else
+	UE_LOG(LogOpenCVHelper, Error, TEXT("DrawCharucoCorners: OpenCV not available (WITH_OPENCV=0)"));
+	return false;
+#endif
+}
+
+
+UTexture2D* FOpenCVHelper::GenerateCharucoBoard(const FCharucoBoardConfig& BoardConfig, FIntPoint ImageSize, int32 MarginSize)
+{
+#if WITH_OPENCV
+
+	// Validate input parameters
+	if (ImageSize.X <= 0 || ImageSize.Y <= 0)
+	{
+		UE_LOG(LogOpenCVHelper, Error, TEXT("GenerateCharucoBoard: Invalid image dimensions (%d x %d). Must be positive."), 
+			ImageSize.X, ImageSize.Y);
+		return nullptr;
+	}
+
+	if (BoardConfig.SquaresX <= 0 || BoardConfig.SquaresY <= 0)
+	{
+		UE_LOG(LogOpenCVHelper, Error, TEXT("GenerateCharucoBoard: Invalid board dimensions (%d x %d squares). Must be positive."), 
+			BoardConfig.SquaresX, BoardConfig.SquaresY);
+		return nullptr;
+	}
+
+	if (BoardConfig.SquareSize <= 0.0f || BoardConfig.MarkerSize <= 0.0f)
+	{
+		UE_LOG(LogOpenCVHelper, Error, TEXT("GenerateCharucoBoard: Invalid sizes (SquareSize: %f, MarkerSize: %f). Must be positive."), 
+			BoardConfig.SquareSize, BoardConfig.MarkerSize);
+		return nullptr;
+	}
+
+	if (BoardConfig.MarkerSize >= BoardConfig.SquareSize)
+	{
+		UE_LOG(LogOpenCVHelper, Error, TEXT("GenerateCharucoBoard: MarkerSize (%f) must be smaller than SquareSize (%f)."), 
+			BoardConfig.MarkerSize, BoardConfig.SquareSize);
+		return nullptr;
+	}
+
+	if (MarginSize < 0)
+	{
+		UE_LOG(LogOpenCVHelper, Error, TEXT("GenerateCharucoBoard: MarginSize (%d) cannot be negative."), MarginSize);
+		return nullptr;
+	}
+
+	// Get ArUco dictionary
+
+	cv::Ptr<cv::aruco::Dictionary> Dictionary = UE::Aruco::Private::GetArucoDictionary(BoardConfig.Dictionary);
+
+	if (Dictionary == nullptr)
+	{
+		UE_LOG(LogOpenCVHelper, Error, TEXT("GenerateCharucoBoard: Invalid ArUco dictionary"));
+		return nullptr;
+	}
+
+	// Create Charuco board
+
+	cv::Ptr<cv::aruco::CharucoBoard> CharucoBoard = cv::aruco::CharucoBoard::create(
+		BoardConfig.SquaresX, 
+		BoardConfig.SquaresY, 
+		BoardConfig.SquareSize, 
+		BoardConfig.MarkerSize, 
+		Dictionary
+	);
+
+	// Generate board image
+
+	cv::Mat BoardImage;
+	CharucoBoard->draw(cv::Size(ImageSize.X, ImageSize.Y), BoardImage, MarginSize, 1);
+
+	// Convert to BGRA
+
+	cv::Mat BoardImageBGRA;
+	cv::cvtColor(BoardImage, BoardImageBGRA, cv::COLOR_GRAY2BGRA);
+
+	// Create UE texture
+
+	UTexture2D* NewTexture = UTexture2D::CreateTransient(ImageSize.X, ImageSize.Y, PF_B8G8R8A8);
+
+	if (!NewTexture)
+	{
+		UE_LOG(LogOpenCVHelper, Error, TEXT("GenerateCharucoBoard: Failed to create texture"));
+		return nullptr;
+	}
+
+	// Copy data to texture
+
+	FTexture2DMipMap& Mip = NewTexture->GetPlatformData()->Mips[0];
+
+	uint8* TextureData = static_cast<uint8*>(Mip.BulkData.Lock(LOCK_READ_WRITE));
+	FMemory::Memcpy(TextureData, BoardImageBGRA.data, BoardImageBGRA.total() * BoardImageBGRA.elemSize());
+
+	Mip.BulkData.Unlock();
+	
+	NewTexture->UpdateResource();
+
+	UE_LOG(LogOpenCVHelper, Log, TEXT("GenerateCharucoBoard: Generated %dx%d Charuco board with %dx%d squares"), 
+		ImageSize.X, ImageSize.Y, BoardConfig.SquaresX, BoardConfig.SquaresY);
+
+	return NewTexture;
+#else
+	UE_LOG(LogOpenCVHelper, Error, TEXT("GenerateCharucoBoard: OpenCV not available"));
+	return nullptr;
+#endif
+}
diff --git a/Engine/Plugins/Runtime/OpenCV/Source/OpenCVHelper/Public/OpenCVHelper.h b/Engine/Plugins/Runtime/OpenCV/Source/OpenCVHelper/Public/OpenCVHelper.h
--- a/Engine/Plugins/Runtime/OpenCV/Source/OpenCVHelper/Public/OpenCVHelper.h
+++ b/Engine/Plugins/Runtime/OpenCV/Source/OpenCVHelper/Public/OpenCVHelper.h
@@ -44,6 +44,7 @@ namespace cv
 
 #endif	// WITH_OPENCV
 
+
 UENUM(BlueprintType)
 enum class EArucoDictionary : uint8
 {
@@ -67,12 +68,48 @@ enum class EArucoDictionary : uint8
 	DICT_ARUCO_ORIGINAL	 UMETA(DisplayName = "DICT_ARUCO_ORIGINAL")
 };
 
+
 struct FArucoMarker
 {
 	int32 MarkerID = 0;
 	FVector2f Corners[4];
 };
 
+
+/** Structure representing a Charuco board configuration */
+struct FCharucoBoardConfig
+{
+	/** Number of chessboard squares in X direction */
+	int32 SquaresX = 5;
+	
+	/** Number of chessboard squares in Y direction */
+	int32 SquaresY = 7;
+	
+	/** Size of each square in the board (in cm) */
+	float SquareSize = 4.0f;
+	
+	/** Size of the ArUco markers (in cm) */
+	float MarkerSize = 3.0f;
+	
+	/** ArUco dictionary to use for the markers. Smaller dictionaries tend to be easier to parse */
+	EArucoDictionary Dictionary = EArucoDictionary::DICT_4X4_50;
+};
+
+
+/** Structure representing detected Charuco corners */
+struct FCharucoCorners
+{
+	/** 2D image coordinates of detected corners */
+	TArray<FVector2f> Corners;
+	
+	/** IDs of the detected corners (based on board layout) */
+	TArray<int32> CornerIds;
+	
+	/** Detected ArUco markers in the Charuco board */
+	TArray<FArucoMarker> Markers;
+};
+
+
 class FOpenCVHelper
 {
 public:
@@ -180,6 +217,15 @@ public:
 
 	/** Project the 3D objects points to the 2D image plane represented by the input camera intrinsics and pose, and compute the reprojection error (euclidean distance) between the input image points and the reprojected points */
 	static OPENCVHELPER_API double ComputeReprojectionError(const TArray<FVector>& ObjectPoints, const TArray<FVector2f>& ImagePoints, const FVector2D& FocalLength, const FVector2D& ImageCenter, const FTransform& CameraPose);
+
+	/** Identify Charuco corners in the input image with the given board configuration, and output the detected corners and their IDs */
+	static OPENCVHELPER_API bool IdentifyCharucoCorners(TArray<FColor>& Image, FIntPoint ImageSize, const FCharucoBoardConfig& BoardConfig, FCharucoCorners& OutCorners, const FIntRect& ROI = FIntRect());
+
+	/** Draw a debug view of the input Charuco corners on top of the input texture */
+	static OPENCVHELPER_API bool DrawCharucoCorners(const FCharucoCorners& Corners, UTexture2D* DebugTexture);
+
+	/** Generate a Charuco board image with the given configuration for display or printing */
+	static OPENCVHELPER_API UTexture2D* GenerateCharucoBoard(const FCharucoBoardConfig& BoardConfig, FIntPoint ImageSize, int32 MarginSize = 10);
 };
 
 /**
diff --git a/Engine/Plugins/VirtualProduction/CameraCalibration/Source/CameraCalibrationEditor/Private/AssetEditor/CameraCalibrationStepsController.cpp b/Engine/Plugins/VirtualProduction/CameraCalibration/Source/CameraCalibrationEditor/Private/AssetEditor/CameraCalibrationStepsController.cpp
--- a/Engine/Plugins/VirtualProduction/CameraCalibration/Source/CameraCalibrationEditor/Private/AssetEditor/CameraCalibrationStepsController.cpp
+++ b/Engine/Plugins/VirtualProduction/CameraCalibration/Source/CameraCalibrationEditor/Private/AssetEditor/CameraCalibrationStepsController.cpp
@@ -1264,22 +1264,31 @@ bool FCameraCalibrationStepsController::CalculateNormalizedMouseClickPosition(co
 
 	if (ViewportPortion == ESimulcamViewportPortion::CameraFeed)
 	{
-		ULensFile* LensFilePtr = GetLensFile();
-		if (!LensFilePtr)
+		FIntRect CropRegion;
+		FIntPoint CropSize;
+		if (!CalculateCameraFeedCropRegion(CropRegion, CropSize))
 		{
 			return false;
 		}
-
-		const FIntPoint CameraFeedDimensions = LensFilePtr->CameraFeedInfo.GetDimensions();
+		
 		const FIntPoint MediaOverlayResolution = GetMediaOverlayResolution();
+		
+		// Convert normalized viewport coordinates to media texture coordinates
+		const float MediaX = XNormalized * MediaOverlayResolution.X;
+		const float MediaY = YNormalized * MediaOverlayResolution.Y;
+		
+		// Check if the click is within the cropped region
+		if (MediaX < CropRegion.Min.X || MediaX > CropRegion.Max.X ||
+		    MediaY < CropRegion.Min.Y || MediaY > CropRegion.Max.Y)
+		{
+			return false;
+		}
+		
+		// Map from cropped region to normalized camera feed coordinates
+		XNormalized = (MediaX - CropRegion.Min.X) / CropSize.X;
+		YNormalized = (MediaY - CropRegion.Min.Y) / CropSize.Y;
 
-		const float AspectRatioCorrectionX = MediaOverlayResolution.X / (float)CameraFeedDimensions.X;
-		const float AspectRatioCorrectionY = MediaOverlayResolution.Y / (float)CameraFeedDimensions.Y;
-
-		XNormalized = ((XNormalized - 0.5f) * AspectRatioCorrectionX) + 0.5f;
-		YNormalized = ((YNormalized - 0.5f) * AspectRatioCorrectionY) + 0.5f;
-
-		// If the scaled values for X or Y are outside the range of [0,1], then the position is invalid (not on the camera feed)
+		// Ensure normalized values are within valid range (defensive check)
 		if (XNormalized < 0.0f || XNormalized > 1.0f || YNormalized < 0.0f || YNormalized > 1.0f)
 		{
 			return false;
@@ -1350,20 +1359,42 @@ bool FCameraCalibrationStepsController::ReadMediaPixels(TArray<FColor>& Pixels,
 	if ((ViewportPortion == ESimulcamViewportPortion::CameraFeed) && LensFilePtr->CameraFeedInfo.IsValid())
 	{
 		Size = LensFilePtr->CameraFeedInfo.GetDimensions();
+		
+		// Protect against invalid dimensions
+		if (Size.X <= 0 || Size.Y <= 0)
+		{
+			OutErrorMessage = LOCTEXT("InvalidCameraFeedDimensions", "Camera feed dimensions are invalid (zero or negative).");
+			return false;
+		}
+		
 		Pixels.SetNumUninitialized(Size.X * Size.Y);
 
 		const FIntPoint MediaResolution = MediaRenderTarget->GetSizeXY();
-		const FIntPoint MediaCenterPoint = MediaResolution / 2;
-		const float AspectRatioCorrectionX = MediaResolution.X / (float)Size.X;
-		const float AspectRatioCorrectionY = MediaResolution.Y / (float)Size.Y;
-
-		// Only return pixels from the actual camera feed (which may be smaller than the full media render target)
+		
+		// Protect against invalid media resolution
+		if (MediaResolution.X <= 0 || MediaResolution.Y <= 0)
+		{
+			OutErrorMessage = LOCTEXT("InvalidMediaResolution", "Media resolution is invalid (zero or negative).");
+			return false;
+		}
+		
+		// Use the common function to calculate crop region
+		FIntRect CropRegion;
+		FIntPoint CropSize;
+		if (!CalculateCameraFeedCropRegion(CropRegion, CropSize))
+		{
+			OutErrorMessage = LOCTEXT("InvalidCropCalculation", "Failed to calculate valid crop region.");
+			return false;
+		}
+		
+		// Map each output pixel to the fitted media region
 		for (int32 YCoordinate = 0; YCoordinate < Size.Y; YCoordinate++)
 		{
 			for (int32 XCoordinate = 0; XCoordinate < Size.X; XCoordinate++)
 			{
-				const int32 ScaledX = (((XCoordinate * AspectRatioCorrectionX) - MediaCenterPoint.X) / AspectRatioCorrectionX) + MediaCenterPoint.X;
-				const int32 ScaledY = (((YCoordinate * AspectRatioCorrectionY) - MediaCenterPoint.Y) / AspectRatioCorrectionY) + MediaCenterPoint.Y;
+				// Map output coordinates to crop region coordinates
+				const int32 ScaledX = CropRegion.Min.X + (XCoordinate * CropSize.X) / Size.X;
+				const int32 ScaledY = CropRegion.Min.Y + (YCoordinate * CropSize.Y) / Size.Y;
 
 				Pixels[YCoordinate * Size.X + XCoordinate] = MediaPixels[ScaledY * MediaResolution.X + ScaledX];
 			}
@@ -1397,4 +1428,49 @@ void FCameraCalibrationStepsController::OnActiveMediaProfileChanged(UMediaProfil
 	ClearMedia();
 }
 
+bool FCameraCalibrationStepsController::CalculateCameraFeedCropRegion(FIntRect& OutCropRegion, FIntPoint& OutCropSize) const
+{
+	ULensFile* LensFilePtr = GetLensFile();
+	if (!LensFilePtr || !LensFilePtr->CameraFeedInfo.IsValid())
+	{
+		return false;
+	}
+
+	const FIntPoint CameraFeedDimensions = LensFilePtr->CameraFeedInfo.GetDimensions();
+	const FIntPoint MediaResolution = GetMediaOverlayResolution();
+	
+	// Protect against invalid dimensions
+	if (CameraFeedDimensions.X <= 0 || CameraFeedDimensions.Y <= 0 || 
+	    MediaResolution.X <= 0 || MediaResolution.Y <= 0)
+	{
+		return false;
+	}
+	
+	// Calculate the center-crop/fit region
+	const float MediaAspect = (float)MediaResolution.X / (float)MediaResolution.Y;
+	const float DesiredAspect = (float)CameraFeedDimensions.X / (float)CameraFeedDimensions.Y;
+	
+	if (MediaAspect > DesiredAspect)
+	{
+		// Media is wider - fit by height, center horizontally
+		OutCropSize.Y = MediaResolution.Y;
+		OutCropSize.X = FMath::RoundToInt(MediaResolution.Y * DesiredAspect);
+		OutCropRegion.Min.X = (MediaResolution.X - OutCropSize.X) / 2;
+		OutCropRegion.Min.Y = 0;
+	}
+	else
+	{
+		// Media is taller - fit by width, center vertically
+		OutCropSize.X = MediaResolution.X;
+		OutCropSize.Y = FMath::RoundToInt(MediaResolution.X / DesiredAspect);
+		OutCropRegion.Min.X = 0;
+		OutCropRegion.Min.Y = (MediaResolution.Y - OutCropSize.Y) / 2;
+	}
+	
+	OutCropRegion.Max = OutCropRegion.Min + OutCropSize;
+	
+	// Validate crop size
+	return (OutCropSize.X > 0 && OutCropSize.Y > 0);
+}
+
 #undef LOCTEXT_NAMESPACE
diff --git a/Engine/Plugins/VirtualProduction/CameraCalibration/Source/CameraCalibrationEditor/Private/AssetEditor/CameraCalibrationStepsController.h b/Engine/Plugins/VirtualProduction/CameraCalibration/Source/CameraCalibrationEditor/Private/AssetEditor/CameraCalibrationStepsController.h
--- a/Engine/Plugins/VirtualProduction/CameraCalibration/Source/CameraCalibrationEditor/Private/AssetEditor/CameraCalibrationStepsController.h
+++ b/Engine/Plugins/VirtualProduction/CameraCalibration/Source/CameraCalibrationEditor/Private/AssetEditor/CameraCalibrationStepsController.h
@@ -264,6 +264,9 @@ private:
 	/** Raised when the active media profile is changed */
 	void OnActiveMediaProfileChanged(UMediaProfile* PreviousMediaProfile, UMediaProfile* NewMediaProfile);
 
+	/** Calculate the crop region for center-fit sampling of camera feed from media texture */
+	bool CalculateCameraFeedCropRegion(FIntRect& OutCropRegion, FIntPoint& OutCropSize) const;
+
 private:
 
 	/** Pointer to the camera calibration toolkit */
diff --git a/Engine/Plugins/VirtualProduction/CameraCalibration/Source/CameraCalibrationEditor/Private/AssetEditor/LensDistortionTool.cpp b/Engine/Plugins/VirtualProduction/CameraCalibration/Source/CameraCalibrationEditor/Private/AssetEditor/LensDistortionTool.cpp
--- a/Engine/Plugins/VirtualProduction/CameraCalibration/Source/CameraCalibrationEditor/Private/AssetEditor/LensDistortionTool.cpp
+++ b/Engine/Plugins/VirtualProduction/CameraCalibration/Source/CameraCalibrationEditor/Private/AssetEditor/LensDistortionTool.cpp
@@ -5,6 +5,7 @@
 #include "AssetToolsModule.h"
 #include "CalibrationPointComponent.h"
 #include "Camera/CameraActor.h"
+#include "CameraCalibrationCharucoBoard.h"
 #include "CameraCalibrationCheckerboard.h"
 #include "CameraCalibrationEditorLog.h"
 #include "CameraCalibrationSettings.h"
@@ -186,15 +187,23 @@ bool ULensDistortionTool::OnViewportClicked(const FGeometry& MyGeometry, const F
 
 bool ULensDistortionTool::OnViewportMarqueeSelect(FVector2D StartPosition, FVector2D EndPosition)
 {
-	// Marquee select currently only supports providing an ROI for checkerboard detection
-	if (CaptureSettings.CalibrationPattern != ECalibrationPattern::Checkerboard)
+	// Marquee select supports ROI for checkerboard and Charuco detection
+	if (CaptureSettings.CalibrationPattern != ECalibrationPattern::Checkerboard && 
+	    CaptureSettings.CalibrationPattern != ECalibrationPattern::Charuco)
 	{
 		return false;
 	}
 
+	// Normalize the marquee selection to ensure Min <= Max for both axes
 	FIntRect MarqueeSelection;
-	MarqueeSelection.Min = FIntPoint(FMath::Floor(StartPosition.X), FMath::Floor(StartPosition.Y));
-	MarqueeSelection.Max = FIntPoint(FMath::Floor(EndPosition.X), FMath::Floor(EndPosition.Y));
+	MarqueeSelection.Min = FIntPoint(
+		FMath::Min(FMath::Floor(StartPosition.X), FMath::Floor(EndPosition.X)),
+		FMath::Min(FMath::Floor(StartPosition.Y), FMath::Floor(EndPosition.Y))
+	);
+	MarqueeSelection.Max = FIntPoint(
+		FMath::Max(FMath::Floor(StartPosition.X), FMath::Floor(EndPosition.X)),
+		FMath::Max(FMath::Floor(StartPosition.Y), FMath::Floor(EndPosition.Y))
+	);
 
 	// These are unused by the checkerboard detection
 	const FGeometry MyGeometryUnused;
@@ -256,6 +265,15 @@ bool ULensDistortionTool::CaptureCalibrationData(const FGeometry& MyGeometry, co
 		{
 			bDetectionResult = DetectArucoPattern(Pixels, ImageSize, NewRow);
 		}
+		else if (CaptureSettings.CalibrationPattern == ECalibrationPattern::Charuco)
+		{
+			// If no ROI was provided, use the full image size
+			if (RegionOfInterest.IsEmpty())
+			{
+				RegionOfInterest = FIntRect(FIntPoint(0), ImageSize);
+			}
+			bDetectionResult = DetectCharucoPattern(Pixels, ImageSize, RegionOfInterest, NewRow);
+		}
 
 		if (!bDetectionResult)
 		{
@@ -463,6 +481,170 @@ bool ULensDistortionTool::DetectArucoPattern(TArray<FColor>& Pixels, FIntPoint S
 	return true;
 }
 
+bool ULensDistortionTool::DetectCharucoPattern(TArray<FColor>& Pixels, FIntPoint Size, FIntRect RegionOfInterest, TSharedPtr<FCalibrationRow> OutRow)
+{
+	const FText ErrorTitle = LOCTEXT("CaptureError", "Capture Error");
+
+	AActor* CalibratorActor = CaptureSettings.Calibrator.Get();
+	if (!CalibratorActor)
+	{
+		const FText ErrorMessage = LOCTEXT("NoCalibratorError", "Please select a valid calibrator actor.");
+		FMessageDialog::Open(EAppMsgCategory::Error, EAppMsgType::Ok, ErrorMessage, ErrorTitle);
+		return false;
+	}
+
+	ACameraCalibrationCharucoBoard* CharucoBoard = Cast<ACameraCalibrationCharucoBoard>(CalibratorActor);
+	if (!CharucoBoard)
+	{
+		const FText ErrorMessage = LOCTEXT("NoCharucoBoardError", "The selected calibrator is not a Charuco board.");
+		FMessageDialog::Open(EAppMsgCategory::Error, EAppMsgType::Ok, ErrorMessage, ErrorTitle);
+		return false;
+	}
+
+	// Setup board configuration
+
+	FCharucoBoardConfig BoardConfig;
+
+	BoardConfig.SquaresX = CharucoBoard->NumCornerCols + 1;
+	BoardConfig.SquaresY = CharucoBoard->NumCornerRows + 1;
+	BoardConfig.SquareSize = CharucoBoard->SquareSideLength;
+	BoardConfig.MarkerSize = CharucoBoard->SquareSideLength * CharucoBoard->MarkerSizeRatio;
+	BoardConfig.Dictionary = CharucoBoard->ArucoDictionary;
+
+	// Handle ROI by cropping the image (if needed)
+
+	TArray<FColor> ProcessedPixels;
+	FIntPoint ProcessedSize;
+	FVector2D ROIOffset(0, 0);
+	
+	if (!RegionOfInterest.IsEmpty() && RegionOfInterest != FIntRect(FIntPoint(0), Size))
+	{
+		// Crop the image to the ROI
+
+		const int32 ROIWidth = RegionOfInterest.Width();
+		const int32 ROIHeight = RegionOfInterest.Height();
+
+		ProcessedPixels.Reserve(ROIWidth * ROIHeight);
+		
+		for (int32 Y = RegionOfInterest.Min.Y; Y < RegionOfInterest.Max.Y; ++Y)
+		{
+			for (int32 X = RegionOfInterest.Min.X; X < RegionOfInterest.Max.X; ++X)
+			{
+				if (X >= 0 && X < Size.X && Y >= 0 && Y < Size.Y)
+				{
+					ProcessedPixels.Add(Pixels[Y * Size.X + X]);
+				}
+				else
+				{
+					ProcessedPixels.Add(FColor::Black);
+				}
+			}
+		}
+		
+		ProcessedSize = FIntPoint(ROIWidth, ROIHeight);
+		ROIOffset = FVector2D(RegionOfInterest.Min.X, RegionOfInterest.Min.Y);
+	}
+	else
+	{
+		// Use the full image
+		ProcessedPixels = Pixels;
+		ProcessedSize = Size;
+	}
+
+	// Detect Charuco corners
+	FCharucoCorners DetectedCorners;
+	bool bResult = FOpenCVHelper::IdentifyCharucoCorners(ProcessedPixels, ProcessedSize, BoardConfig, DetectedCorners);
+	
+	// Adjust detected coordinates by ROI offset
+	if (bResult && ROIOffset != FVector2D::ZeroVector)
+	{
+		for (FVector2f& Corner : DetectedCorners.Corners)
+		{
+			Corner.X += ROIOffset.X;
+			Corner.Y += ROIOffset.Y;
+		}
+		for (FArucoMarker& Marker : DetectedCorners.Markers)
+		{
+			for (int32 CornerIdx = 0; CornerIdx < 4; ++CornerIdx)
+			{
+				Marker.Corners[CornerIdx].X += ROIOffset.X;
+				Marker.Corners[CornerIdx].Y += ROIOffset.Y;
+			}
+		}
+	}
+
+	if (!bResult || DetectedCorners.Corners.IsEmpty())
+	{
+		const FText ErrorMessage = FText::FromString(FString::Printf(TEXT(
+			"Failed to detect Charuco corners in the image. "
+			"Make sure the board is clearly visible and well-lit. "
+			"Expected board: %dx%d squares with ArUco dictionary %s"),
+			BoardConfig.SquaresX, BoardConfig.SquaresY,
+			*UEnum::GetValueAsString(BoardConfig.Dictionary)
+		));
+
+		FMessageDialog::Open(EAppMsgCategory::Error, EAppMsgType::Ok, ErrorMessage, ErrorTitle);
+
+		return false;
+	}
+
+	// Convert detected corners to calibration row format
+	//
+
+	// Corners and CornerIds arrays must have the same size
+	if (DetectedCorners.Corners.Num() != DetectedCorners.CornerIds.Num())
+	{
+		const FText ErrorMessage = FText::Format(LOCTEXT("CornerDataMismatchError", 
+			"Charuco detection data is corrupted: {0} corners but {1} corner IDs. Data integrity check failed."),
+			DetectedCorners.Corners.Num(), DetectedCorners.CornerIds.Num());
+		FMessageDialog::Open(EAppMsgCategory::Error, EAppMsgType::Ok, ErrorMessage, ErrorTitle);
+		return false;
+	}
+
+	for (int32 CornerIdx = 0; CornerIdx < DetectedCorners.Corners.Num(); ++CornerIdx)
+	{
+		int32 CornerId = DetectedCorners.CornerIds[CornerIdx];
+		FVector2f ImageCorner = DetectedCorners.Corners[CornerIdx];
+		
+		// Get the 3D position of this corner from the board
+		FVector WorldCorner = CharucoBoard->GetCornerPosition(CornerId);
+				
+		OutRow->ObjectPoints.Points.Add(WorldCorner);
+		OutRow->ImagePoints.Points.Add(FVector2D(ImageCorner.X, ImageCorner.Y));
+	}
+
+	// Set target pose from the board transform
+	OutRow->TargetPose = CharucoBoard->GetActorTransform();
+
+	// Rescale corners to match overlay texture size before drawing
+	if (OverlayTexture)
+	{
+		const FIntPoint OverlayTextureSize = FIntPoint(OverlayTexture->GetSizeX(), OverlayTexture->GetSizeY());
+		const FVector2D TopLeftCorner = (OverlayTextureSize - Size) / 2.0f;
+		
+		// Create a copy of detected corners with rescaled positions
+		FCharucoCorners RescaledCorners = DetectedCorners;
+		for (FVector2f& Corner : RescaledCorners.Corners)
+		{
+			Corner.X += TopLeftCorner.X;
+			Corner.Y += TopLeftCorner.Y;
+		}
+		for (FArucoMarker& Marker : RescaledCorners.Markers)
+		{
+			for (int32 CornerIdx = 0; CornerIdx < 4; ++CornerIdx)
+			{
+				Marker.Corners[CornerIdx].X += TopLeftCorner.X;
+				Marker.Corners[CornerIdx].Y += TopLeftCorner.Y;
+			}
+		}
+		
+		// Draw the rescaled corners on overlay
+		FOpenCVHelper::DrawCharucoCorners(RescaledCorners, OverlayTexture);
+	}
+
+	return true;
+}
+
 bool ULensDistortionTool::DetectPoint(const FGeometry& MyGeometry, const FPointerEvent& MouseEvent, TSharedPtr<FCalibrationRow> OutRow)
 {
 	TSharedPtr<FCameraCalibrationStepsController> StepsController = WeakStepsController.Pin();
diff --git a/Engine/Plugins/VirtualProduction/CameraCalibration/Source/CameraCalibrationEditor/Private/AssetEditor/LensDistortionTool.h b/Engine/Plugins/VirtualProduction/CameraCalibration/Source/CameraCalibrationEditor/Private/AssetEditor/LensDistortionTool.h
--- a/Engine/Plugins/VirtualProduction/CameraCalibration/Source/CameraCalibrationEditor/Private/AssetEditor/LensDistortionTool.h
+++ b/Engine/Plugins/VirtualProduction/CameraCalibration/Source/CameraCalibrationEditor/Private/AssetEditor/LensDistortionTool.h
@@ -26,7 +26,8 @@ enum class ECalibrationPattern : uint8
 {
 	Checkerboard,
 	Aruco,
-	Points
+	Points,
+	Charuco,
 };
 
 /** Version info to support backwards compatibility of dataset importing */
@@ -93,7 +94,7 @@ struct FLensCaptureSettings
 {
 	GENERATED_BODY()
 
-	/** The pattern to detect in the media image */
+	/** The pattern to detect in the media image. CharUco is recommended when available. */
 	UPROPERTY(EditAnywhere, Category = "Capture Settings")
 	ECalibrationPattern CalibrationPattern = ECalibrationPattern::Checkerboard;
 
@@ -226,6 +227,9 @@ private:
 	/** Detects an aruco pattern in the input array of pixels (supplied by the media image) and fills out the 2D and 3D points in the output calibration row */
 	bool DetectArucoPattern(TArray<FColor>& Pixels, FIntPoint Size, TSharedPtr<FCalibrationRow> OutRow);
 
+	/** Detects a charuco pattern in the input array of pixels (supplied by the media image) and fills out the 2D and 3D points in the output calibration row */
+	bool DetectCharucoPattern(TArray<FColor>& Pixels, FIntPoint Size, FIntRect RegionOfInterest, TSharedPtr<FCalibrationRow> OutRow);
+
 	/** Captures the mouse position and the 3D world location of the current calibration component and and fills out the output calibration row */
 	bool DetectPoint(const FGeometry& MyGeometry, const FPointerEvent& MouseEvent, TSharedPtr<FCalibrationRow> OutRow);
 
diff --git a/Engine/Plugins/VirtualProduction/CameraCalibration/Source/CameraCalibrationEditor/Private/AssetEditor/NodalOffsetTool.cpp b/Engine/Plugins/VirtualProduction/CameraCalibration/Source/CameraCalibrationEditor/Private/AssetEditor/NodalOffsetTool.cpp
--- a/Engine/Plugins/VirtualProduction/CameraCalibration/Source/CameraCalibrationEditor/Private/AssetEditor/NodalOffsetTool.cpp
+++ b/Engine/Plugins/VirtualProduction/CameraCalibration/Source/CameraCalibrationEditor/Private/AssetEditor/NodalOffsetTool.cpp
@@ -88,6 +88,17 @@ bool UNodalOffsetTool::OnViewportClicked(const FGeometry& MyGeometry, const FPoi
 	return NodalOffsetAlgo->OnViewportClicked(MyGeometry, MouseEvent);
 }
 
+bool UNodalOffsetTool::OnViewportMarqueeSelect(FVector2D StartPosition, FVector2D EndPosition)
+{
+	if (!bIsActive || !NodalOffsetAlgo)
+	{
+		return false;
+	}
+
+	// Each algo has its own reaction to the marquee selection.
+	return NodalOffsetAlgo->OnViewportMarqueeSelect(StartPosition, EndPosition);
+}
+
 bool UNodalOffsetTool::OnViewportInputKey(const FKey& InKey, const EInputEvent& InEvent)
 {
 	if (!bIsActive)
diff --git a/Engine/Plugins/VirtualProduction/CameraCalibration/Source/CameraCalibrationEditor/Private/AssetEditor/NodalOffsetTool.h b/Engine/Plugins/VirtualProduction/CameraCalibration/Source/CameraCalibrationEditor/Private/AssetEditor/NodalOffsetTool.h
--- a/Engine/Plugins/VirtualProduction/CameraCalibration/Source/CameraCalibrationEditor/Private/AssetEditor/NodalOffsetTool.h
+++ b/Engine/Plugins/VirtualProduction/CameraCalibration/Source/CameraCalibrationEditor/Private/AssetEditor/NodalOffsetTool.h
@@ -47,6 +47,7 @@ public:
 	virtual void Shutdown() override;
 	virtual void Tick(float DeltaTime) override;
 	virtual bool OnViewportClicked(const FGeometry& MyGeometry, const FPointerEvent& MouseEvent) override;
+	virtual bool OnViewportMarqueeSelect(FVector2D StartPosition, FVector2D EndPosition) override;
 	virtual bool OnViewportInputKey(const FKey& InKey, const EInputEvent& InEvent) override;
 	virtual TSharedRef<SWidget> BuildUI() override;
 	virtual FName FriendlyName() const  override { return TEXT("Nodal Offset"); };
diff --git a/Engine/Plugins/VirtualProduction/CameraCalibration/Source/CameraCalibrationEditor/Private/AssetEditor/Viewport/SimulcamEditorViewportClient.cpp b/Engine/Plugins/VirtualProduction/CameraCalibration/Source/CameraCalibrationEditor/Private/AssetEditor/Viewport/SimulcamEditorViewportClient.cpp
--- a/Engine/Plugins/VirtualProduction/CameraCalibration/Source/CameraCalibrationEditor/Private/AssetEditor/Viewport/SimulcamEditorViewportClient.cpp
+++ b/Engine/Plugins/VirtualProduction/CameraCalibration/Source/CameraCalibrationEditor/Private/AssetEditor/Viewport/SimulcamEditorViewportClient.cpp
@@ -625,7 +625,7 @@ FText FSimulcamEditorViewportClient::GetDisplayedResolution() const
 	}
 
 	return FText::Format(
-		FText::FromString("Displayed: {0}x{1}\nTextureSize: {2}x{3}\nTexturePosition: {4}x{5}\nZoom: {6}"),
+		FText::FromString("Displayed: {0}x{1}\nTextureSize: {2}x{3}\nTexturePosition: {4}x{5}\nZoom: {6}\n\nCtrl+Alt+Drag: Marquee ROI selection"),
 		FText::AsNumber(ScaledTextureSize.X),
 		FText::AsNumber(ScaledTextureSize.Y),
 		FText::AsNumber(TextureSize.X),
diff --git a/Engine/Plugins/VirtualProduction/CameraCalibration/Source/CameraCalibrationEditor/Private/CameraCalibrationEditorModule.cpp b/Engine/Plugins/VirtualProduction/CameraCalibration/Source/CameraCalibrationEditor/Private/CameraCalibrationEditorModule.cpp
--- a/Engine/Plugins/VirtualProduction/CameraCalibration/Source/CameraCalibrationEditor/Private/CameraCalibrationEditorModule.cpp
+++ b/Engine/Plugins/VirtualProduction/CameraCalibration/Source/CameraCalibrationEditor/Private/CameraCalibrationEditorModule.cpp
@@ -8,6 +8,7 @@
 #include "AssetEditor/Curves/LensDataCurveModel.h"
 #include "AssetEditor/SCameraCalibrationCurveEditorView.h"
 #include "AssetToolsModule.h"
+#include "CameraCalibrationCharucoBoard.h"
 #include "CameraCalibrationEditorLog.h"
 #include "CameraCalibrationSubsystem.h"
 #include "CameraCalibrationTypes.h"
@@ -133,6 +134,19 @@ void FCameraCalibrationEditorModule::RegisterPlacementModeItems()
 				NSLOCTEXT("PlacementMode", "TrackerV3", "TrackerV3")
 				)));
 		}
+
+		// Register the Charuco Board
+		{
+			PlaceActors.Add(IPlacementModeModule::Get().RegisterPlaceableItem(Info->UniqueHandle, MakeShared<FPlaceableItem>(
+				*UActorFactory::StaticClass(),
+				FAssetData(ACameraCalibrationCharucoBoard::StaticClass()),
+				NAME_None,
+				NAME_None,
+				TOptional<FLinearColor>(),
+				TOptional<int32>(),
+				NSLOCTEXT("PlacementMode", "CharucoBoard", "Charuco Board")
+			)));
+		}
 	};
 
 	if (FApp::CanEverRender())
diff --git a/Engine/Plugins/VirtualProduction/CameraCalibrationCore/CameraCalibrationCore.uplugin b/Engine/Plugins/VirtualProduction/CameraCalibrationCore/CameraCalibrationCore.uplugin
--- a/Engine/Plugins/VirtualProduction/CameraCalibrationCore/CameraCalibrationCore.uplugin
+++ b/Engine/Plugins/VirtualProduction/CameraCalibrationCore/CameraCalibrationCore.uplugin
@@ -34,6 +34,10 @@
 			"Name": "GameplayCameras",
 			"Enabled": true
 		},
+		{
+			"Name": "OpenCV",
+			"Enabled": true
+		},
 		{
 			"Name": "ProceduralMeshComponent",
 			"Enabled": true
diff --git a/Engine/Plugins/VirtualProduction/CameraCalibrationCore/Source/CameraCalibrationCore/CameraCalibrationCore.Build.cs b/Engine/Plugins/VirtualProduction/CameraCalibrationCore/Source/CameraCalibrationCore/CameraCalibrationCore.Build.cs
--- a/Engine/Plugins/VirtualProduction/CameraCalibrationCore/Source/CameraCalibrationCore/CameraCalibrationCore.Build.cs
+++ b/Engine/Plugins/VirtualProduction/CameraCalibrationCore/Source/CameraCalibrationCore/CameraCalibrationCore.Build.cs
@@ -15,6 +15,8 @@ namespace UnrealBuildTool.Rules
 					"DeveloperSettings",
 					"Engine",
 					"GameplayCameras",
+					"OpenCV",
+					"OpenCVHelper",
 					"RenderCore",
 					"RHI",
 					"Slate",
diff --git a/Engine/Plugins/VirtualProduction/CameraCalibrationCore/Source/CameraCalibrationCore/Private/CameraCalibrationCharucoBoard.cpp b/Engine/Plugins/VirtualProduction/CameraCalibrationCore/Source/CameraCalibrationCore/Private/CameraCalibrationCharucoBoard.cpp
new file mode 100644
--- /dev/null
+++ b/Engine/Plugins/VirtualProduction/CameraCalibrationCore/Source/CameraCalibrationCore/Private/CameraCalibrationCharucoBoard.cpp
@@ -0,0 +1,262 @@
+// Copyright Epic Games, Inc. All Rights Reserved.
+
+#include "CameraCalibrationCharucoBoard.h"
+#include "CameraCalibrationCoreLog.h"
+#include "CalibrationPointComponent.h"
+#include "Components/StaticMeshComponent.h"
+#include "Engine/StaticMesh.h"
+#include "Engine/Texture2D.h"
+#include "Materials/MaterialInstanceDynamic.h"
+#include "Materials/MaterialInterface.h"
+#include "OpenCVHelper.h"
+#include "UObject/ConstructorHelpers.h"
+
+
+ACameraCalibrationCharucoBoard::ACameraCalibrationCharucoBoard()
+{
+	// Create root component
+	Root = CreateDefaultSubobject<USceneComponent>(TEXT("Root"));
+	SetRootComponent(Root);
+	RootComponent->SetVisibility(true);
+
+	// Create board mesh component
+	BoardMeshComponent = CreateDefaultSubobject<UStaticMeshComponent>(TEXT("BoardMesh"));
+	BoardMeshComponent->SetupAttachment(RootComponent);
+
+	// Create fixed calibration point components
+	TopLeft     = CreateDefaultSubobject<UCalibrationPointComponent>(TEXT("TopLeft"));
+	TopRight    = CreateDefaultSubobject<UCalibrationPointComponent>(TEXT("TopRight"));
+	BottomLeft  = CreateDefaultSubobject<UCalibrationPointComponent>(TEXT("BottomLeft"));
+	BottomRight = CreateDefaultSubobject<UCalibrationPointComponent>(TEXT("BottomRight"));
+	Center      = CreateDefaultSubobject<UCalibrationPointComponent>(TEXT("Center"));
+
+	TopLeft    ->SetupAttachment(RootComponent);
+	TopRight   ->SetupAttachment(RootComponent);
+	BottomLeft ->SetupAttachment(RootComponent);
+	BottomRight->SetupAttachment(RootComponent);
+	Center     ->SetupAttachment(RootComponent);
+
+	// Set default ArUco dictionary. Smaller dictionaries tend to be easier to recongnize.
+	ArucoDictionary = EArucoDictionary::DICT_4X4_50;
+
+	// Try to find default cube mesh
+	static ConstructorHelpers::FObjectFinder<UStaticMesh> CubeMeshFinder(TEXT("/Engine/BasicShapes/Cube"));
+	if (CubeMeshFinder.Succeeded())
+	{
+		BoardMesh = CubeMeshFinder.Object;
+	}
+
+	// Find the Charuco board material
+	static ConstructorHelpers::FObjectFinder<UMaterialInterface> MaterialFinder(TEXT("/CameraCalibrationCore/Materials/M_CharucoBoard"));
+	if (MaterialFinder.Succeeded())
+	{
+		BoardMaterial = MaterialFinder.Object;
+	}
+	else
+	{
+		// Unexpected fallback
+
+		UE_LOG(LogCameraCalibrationCore, Error, TEXT("Failed to load M_CharucoBoard"));
+
+		static ConstructorHelpers::FObjectFinder<UMaterialInterface> FallbackFinder(TEXT("/Engine/BasicShapes/BasicShapeMaterial"));
+		if (FallbackFinder.Succeeded())
+		{
+			BoardMaterial = FallbackFinder.Object;
+		}
+	}
+}
+
+void ACameraCalibrationCharucoBoard::OnConstruction(const FTransform& Transform)
+{
+	Super::OnConstruction(Transform);
+	Rebuild();
+}
+
+#if WITH_EDITOR
+void ACameraCalibrationCharucoBoard::PostEditChangeProperty(FPropertyChangedEvent& PropertyChangedEvent)
+{
+	Super::PostEditChangeProperty(PropertyChangedEvent);
+	
+	if (PropertyChangedEvent.Property)
+	{
+		const FName PropertyName = PropertyChangedEvent.Property->GetFName();
+		
+		// Rebuild when relevant properties change
+		if (PropertyName == GET_MEMBER_NAME_CHECKED(ACameraCalibrationCharucoBoard, NumCornerRows) ||
+			PropertyName == GET_MEMBER_NAME_CHECKED(ACameraCalibrationCharucoBoard, NumCornerCols) ||
+			PropertyName == GET_MEMBER_NAME_CHECKED(ACameraCalibrationCharucoBoard, SquareSideLength) ||
+			PropertyName == GET_MEMBER_NAME_CHECKED(ACameraCalibrationCharucoBoard, MarkerSizeRatio) ||
+			PropertyName == GET_MEMBER_NAME_CHECKED(ACameraCalibrationCharucoBoard, ArucoDictionary) ||
+			PropertyName == GET_MEMBER_NAME_CHECKED(ACameraCalibrationCharucoBoard, Thickness) ||
+			PropertyName == GET_MEMBER_NAME_CHECKED(ACameraCalibrationCharucoBoard, BoardMesh) ||
+			PropertyName == GET_MEMBER_NAME_CHECKED(ACameraCalibrationCharucoBoard, BoardMaterial) ||
+			PropertyName == GET_MEMBER_NAME_CHECKED(ACameraCalibrationCharucoBoard, MaxTextureResolution) ||
+			PropertyName == GET_MEMBER_NAME_CHECKED(ACameraCalibrationCharucoBoard, bDisplayPattern))
+		{
+			Rebuild();
+		}
+	}
+}
+#endif
+
+void ACameraCalibrationCharucoBoard::Rebuild()
+{
+	Modify();
+
+	CreateBoardMesh();
+	GenerateCharucoTexture();
+	UpdateCalibrationPointPositions();
+}
+
+void ACameraCalibrationCharucoBoard::ClearGeneratedMeshComponent()
+{
+	// Remove board mesh component
+	if (BoardMeshComponent)
+	{
+		RemoveInstanceComponent(BoardMeshComponent);
+		BoardMeshComponent->DestroyComponent();
+		BoardMeshComponent = nullptr;
+	}
+}
+
+void ACameraCalibrationCharucoBoard::CreateBoardMesh()
+{
+	if (!BoardMesh || !BoardMeshComponent)
+	{
+		return;
+	}
+
+	// Update the mesh and properties
+	BoardMeshComponent->SetStaticMesh(BoardMesh);
+
+	// Calculate board size (NumCorners + 1 squares in each direction)
+	float BoardWidth = (NumCornerCols + 1) * SquareSideLength;
+	float BoardHeight = (NumCornerRows + 1) * SquareSideLength;
+
+	// Scale and position the mesh so origin is at the outer corner
+	BoardMeshComponent->SetRelativeScale3D(FVector(Thickness / 100.0f, BoardWidth / 100.0f, BoardHeight / 100.0f));
+	
+	// Position board so origin (0,0,0) is at bottom-left outer corner
+	// The cube mesh has its origin at center, so offset by half the scaled dimensions  
+	BoardMeshComponent->SetRelativeLocation(FVector(0.0f, BoardWidth / 2.0f, BoardHeight / 2.0f));
+
+	// Apply material if pattern display is enabled
+	if (bDisplayPattern && BoardMaterial)
+	{
+		DynamicBoardMaterial = UMaterialInstanceDynamic::Create(BoardMaterial, this);
+		BoardMeshComponent->SetMaterial(0, DynamicBoardMaterial);
+	}
+}
+
+void ACameraCalibrationCharucoBoard::GenerateCharucoTexture()
+{
+	// Setup board configuration (NumCorners + 1 == squares in each direction)
+
+	FCharucoBoardConfig Config;
+
+	Config.SquaresX = NumCornerCols + 1;
+	Config.SquaresY = NumCornerRows + 1;
+	Config.SquareSize = SquareSideLength;
+	Config.MarkerSize = SquareSideLength * MarkerSizeRatio;
+	Config.Dictionary = ArucoDictionary;
+
+	// Calculate aspect ratio for texture resolution
+	FIntPoint AdjustedResolution(MaxTextureResolution, MaxTextureResolution);
+	
+	if (Config.SquaresY > 0)
+	{
+		float AspectRatio = float(Config.SquaresX) / float(Config.SquaresY);
+		
+		if (AspectRatio >= 1.0f)
+		{
+			// Wider than tall
+			AdjustedResolution.X = MaxTextureResolution;
+			AdjustedResolution.Y = FMath::RoundToInt(MaxTextureResolution / AspectRatio);
+		}
+		else if (AspectRatio > 0.0f)
+		{
+			// Taller than wide
+			AdjustedResolution.Y = MaxTextureResolution;
+			AdjustedResolution.X = FMath::RoundToInt(MaxTextureResolution * AspectRatio);
+		}
+	}
+
+	// Generate the Charuco board texture with no margin and correct aspect ratio
+	CharucoBoardTexture = FOpenCVHelper::GenerateCharucoBoard(Config, AdjustedResolution, 0);
+	
+	if (!CharucoBoardTexture)
+	{
+		UE_LOG(LogCameraCalibrationCore, Error, TEXT("Failed to create Charuco texture"));
+		return;
+	}
+
+	// Apply texture to material if available
+	if (CharucoBoardTexture && DynamicBoardMaterial && bDisplayPattern)
+	{
+		DynamicBoardMaterial->SetTextureParameterValue(TEXT("Texture"), CharucoBoardTexture);
+	}
+	else
+	{
+		UE_LOG(LogCameraCalibrationCore, Error, TEXT("Cannot apply texture - CharucoBoardTexture: %s, DynamicBoardMaterial: %s, bDisplayPattern: %s"),
+			CharucoBoardTexture ? TEXT("Valid") : TEXT("NULL"),
+			DynamicBoardMaterial ? TEXT("Valid") : TEXT("NULL"),
+			bDisplayPattern ? TEXT("True") : TEXT("False"));
+	}
+}
+
+void ACameraCalibrationCharucoBoard::UpdateCalibrationPointPositions()
+{
+	// Position the calibration points at the inner corners.
+	// Board origin is now at bottom-left outer corner, so inner corners are simply offset by SquareSideLength
+	
+	if (TopLeft)
+	{
+		TopLeft->SetRelativeLocation(FVector(0, SquareSideLength, (NumCornerRows) * SquareSideLength));
+	}
+
+	if (TopRight)
+	{
+		TopRight->SetRelativeLocation(FVector(0, NumCornerCols * SquareSideLength, (NumCornerRows) * SquareSideLength));
+	}
+
+	if (BottomLeft)
+	{
+		BottomLeft->SetRelativeLocation(FVector(0, SquareSideLength, SquareSideLength));
+	}
+
+	if (BottomRight)
+	{
+		BottomRight->SetRelativeLocation(FVector(0, NumCornerCols * SquareSideLength, SquareSideLength));
+	}
+
+	if (Center)
+	{
+		// Center of the inner corner grid
+		float CenterY = (NumCornerCols + 1) * 0.5f * SquareSideLength;
+		float CenterZ = (NumCornerRows + 1) * 0.5f * SquareSideLength;
+		Center->SetRelativeLocation(FVector(0, CenterY, CenterZ));
+	}
+}
+
+int32 ACameraCalibrationCharucoBoard::GetNumCorners() const
+{
+	return NumCornerCols * NumCornerRows;
+}
+
+FVector ACameraCalibrationCharucoBoard::GetCornerPosition(int32 CornerId) const
+{
+	if (CornerId < 0 || CornerId >= NumCornerCols * NumCornerRows)
+	{
+		return FVector::ZeroVector;
+	}
+
+	int32 Row = CornerId / NumCornerCols;
+	int32 Col = CornerId % NumCornerCols;
+
+	// Calculate local position using the board's coordinate system
+	// Inner corners start at (1,1), that is, offset by one square from board origin
+	FVector LocalPosition = SquareSideLength * FVector(0, Col + 1, Row + 1);
+
+	// Transform to world coordinates
+	return GetActorTransform().TransformPosition(LocalPosition);
+}
\ No newline at end of file
diff --git a/Engine/Plugins/VirtualProduction/CameraCalibrationCore/Source/CameraCalibrationCore/Public/CameraCalibrationCharucoBoard.h b/Engine/Plugins/VirtualProduction/CameraCalibrationCore/Source/CameraCalibrationCore/Public/CameraCalibrationCharucoBoard.h
new file mode 100644
--- /dev/null
+++ b/Engine/Plugins/VirtualProduction/CameraCalibrationCore/Source/CameraCalibrationCore/Public/CameraCalibrationCharucoBoard.h
@@ -0,0 +1,147 @@
+// Copyright Epic Games, Inc. All Rights Reserved.
+
+#pragma once
+
+#include "CoreTypes.h"
+#include "GameFramework/Actor.h"
+
+#include "CameraCalibrationCharucoBoard.generated.h"
+
+#define UE_API CAMERACALIBRATIONCORE_API
+
+class UCalibrationPointComponent;
+class UMaterialInterface;
+class UMaterialInstanceDynamic;
+class UStaticMesh;
+class UTexture2D;
+
+enum class EArucoDictionary : uint8;
+
+/**
+ * Dynamic Charuco board actor for camera calibration
+ */
+UCLASS(MinimalAPI)
+class ACameraCalibrationCharucoBoard : public AActor
+{
+	GENERATED_BODY()
+
+public:
+	UE_API ACameraCalibrationCharucoBoard();
+
+public:
+	/** Rebuilds the board mesh and generates the Charuco pattern texture */
+	UFUNCTION(BlueprintCallable, Category = "Calibration")
+	UE_API void Rebuild();
+
+	/** Get the generated Charuco board texture */
+	UFUNCTION(BlueprintCallable, Category = "Calibration")
+	UE_API UTexture2D* GetCharucoBoardTexture() const { return CharucoBoardTexture; }
+
+	/** Get the total number of corners in the Charuco board */
+	UFUNCTION(BlueprintCallable, Category = "Calibration")
+	UE_API int32 GetNumCorners() const;
+
+	/** Get the 3D position of a specific corner by its ID */
+	UFUNCTION(BlueprintCallable, Category = "Calibration")
+	UE_API FVector GetCornerPosition(int32 CornerId) const;
+
+	//~ Begin UObject Interface
+	UE_API virtual void OnConstruction(const FTransform& Transform) override;
+#if WITH_EDITOR
+	UE_API virtual void PostEditChangeProperty(FPropertyChangedEvent& PropertyChangedEvent) override;
+#endif
+	//~ End UObject Interface
+
+private:
+	/** Creates the board mesh */
+	void CreateBoardMesh();
+	
+	/** Generates the Charuco pattern texture */
+	void GenerateCharucoTexture();
+	
+	/** Positions the calibration point components */
+	void UpdateCalibrationPointPositions();
+	
+	/** Clears generated mesh component */
+	void ClearGeneratedMeshComponent();
+
+public:
+	/** Root component, gives the Actor a transform */
+	UPROPERTY(EditAnywhere, Category = "Calibration")
+	TObjectPtr<USceneComponent> Root;
+
+	/** TopLeft calibration point */
+	UPROPERTY(EditAnywhere, Category = "Calibration")
+	TObjectPtr<UCalibrationPointComponent> TopLeft;
+
+	/** TopRight calibration point */
+	UPROPERTY(EditAnywhere, Category = "Calibration")
+	TObjectPtr<UCalibrationPointComponent> TopRight;
+
+	/** BottomLeft calibration point */
+	UPROPERTY(EditAnywhere, Category = "Calibration")
+	TObjectPtr<UCalibrationPointComponent> BottomLeft;
+
+	/** BottomRight calibration point */
+	UPROPERTY(EditAnywhere, Category = "Calibration")
+	TObjectPtr<UCalibrationPointComponent> BottomRight;
+
+	/** Center calibration point */
+	UPROPERTY(EditAnywhere, Category = "Calibration")
+	TObjectPtr<UCalibrationPointComponent> Center;
+
+	/** Number of inner corner rows (same convention as checkerboard) */
+	UPROPERTY(EditAnywhere, BlueprintReadWrite, Category = "Calibration", meta = (ClampMin = "2", ClampMax = "63"))
+	int32 NumCornerRows = 5;
+
+	/** Number of inner corner columns (same convention as checkerboard) */
+	UPROPERTY(EditAnywhere, BlueprintReadWrite, Category = "Calibration", meta = (ClampMin = "2", ClampMax = "63"))
+	int32 NumCornerCols = 7;
+
+	/** Length of the side of each square (in cm) */
+	UPROPERTY(EditAnywhere, BlueprintReadWrite, Category = "Calibration", meta = (ClampMin = "1.0", ClampMax = "100.0"))
+	float SquareSideLength = 4.0f;
+
+	/** Size of the ArUco markers relative to square size (0.0 - 1.0) */
+	UPROPERTY(EditAnywhere, BlueprintReadWrite, Category = "Calibration", meta = (ClampMin = "0.1", ClampMax = "0.9"))
+	float MarkerSizeRatio = 0.75f;
+
+	/** ArUco dictionary to use for markers */
+	UPROPERTY(EditAnywhere, BlueprintReadWrite, Category = "Calibration")
+	EArucoDictionary ArucoDictionary;
+
+	/** Thickness of the board (not used for calibration) */
+	UPROPERTY(EditAnywhere, BlueprintReadWrite, Category = "Calibration", meta = (ClampMin = "0.01", ClampMax = "10.0"))
+	float Thickness = 0.1f;
+
+	/** The static mesh to use for the board */
+	UPROPERTY(EditAnywhere, Category = "Calibration")
+	TObjectPtr<UStaticMesh> BoardMesh;
+
+	/** The material to apply the Charuco pattern to */
+	UPROPERTY(EditAnywhere, Category = "Calibration")
+	TObjectPtr<UMaterialInterface> BoardMaterial;
+
+	/** Maximum dimension of the generated texture. The actual resolution will maintain the board's aspect ratio. */
+	UPROPERTY(EditAnywhere, Category = "Calibration", meta = (ClampMin = "256", ClampMax = "4096", UIMin = "256", UIMax = "4096"))
+	int32 MaxTextureResolution = 2048;
+
+	/** Display the generated Charuco pattern on the board mesh */
+	UPROPERTY(EditAnywhere, Category = "Calibration")
+	bool bDisplayPattern = true;
+
+private:
+	/** Generated texture containing the Charuco pattern */
+	UPROPERTY(Transient)
+	TObjectPtr<UTexture2D> CharucoBoardTexture;
+
+	/** Mesh component for the board */
+	UPROPERTY(VisibleAnywhere, Category = "Calibration")
+	TObjectPtr<UStaticMeshComponent> BoardMeshComponent;
+
+	/** Dynamic material instance for displaying the pattern */
+	UPROPERTY(Transient)
+	TObjectPtr<UMaterialInstanceDynamic> DynamicBoardMaterial;
+};
+
+#undef UE_API
\ No newline at end of file
diff --git a/Engine/Plugins/VirtualProduction/CameraCalibrationCore/Source/CameraCalibrationCore/Public/CameraNodalOffsetAlgo.h b/Engine/Plugins/VirtualProduction/CameraCalibrationCore/Source/CameraCalibrationCore/Public/CameraNodalOffsetAlgo.h
--- a/Engine/Plugins/VirtualProduction/CameraCalibrationCore/Source/CameraCalibrationCore/Public/CameraNodalOffsetAlgo.h
+++ b/Engine/Plugins/VirtualProduction/CameraCalibrationCore/Source/CameraCalibrationCore/Public/CameraNodalOffsetAlgo.h
@@ -46,6 +46,9 @@ public:
 	/** Callback when viewport receives input key presses. Returns false if the event was not handled. */
 	virtual bool OnViewportInputKey(const FKey& InKey, const EInputEvent& InEvent) { return false; };
 
+	/** Callback when the user performs a marquee select operation. Returns false if the event was not handled. */
+	virtual bool OnViewportMarqueeSelect(FVector2D StartPosition, FVector2D EndPosition) { return false; };
+
 	/** Returns the UI of this calibrator. Expected to only be called once */
 	virtual TSharedRef<SWidget> BuildUI() { return SNullWidget::NullWidget; };
 
