Fix for crash when OpenCV fails to load

#rnx

diff --git a/Engine/Plugins/Runtime/OpenCV/Source/OpenCVHelper/Private/OpenCVHelperModule.cpp b/Engine/Plugins/Runtime/OpenCV/Source/OpenCVHelper/Private/OpenCVHelperModule.cpp
--- a/Engine/Plugins/Runtime/OpenCV/Source/OpenCVHelper/Private/OpenCVHelperModule.cpp
+++ b/Engine/Plugins/Runtime/OpenCV/Source/OpenCVHelper/Private/OpenCVHelperModule.cpp
@@ -83,6 +83,11 @@ void FOpenCVHelperModule::StartupModule()
 	FPlatformProcess::PushDllDirectory(*OpenCvBinPath);
 	OpenCvDllHandle = FPlatformProcess::GetDllHandle(*DLLPath);
 	FPlatformProcess::PopDllDirectory(*OpenCvBinPath);
+
+	if (OpenCvDllHandle == nullptr)
+	{
+		return;
+	}
 #endif
 
 	// We need to tell OpenCV to use Unreal's memory allocator to avoid crashes.
